#android Faster CPU performance by setting affinity mask for the gamethread, renderthread and eventthread.

diff --git a/Engine/Source/Runtime/Core/Private/Android/AndroidProcess.cpp b/Engine/Source/Runtime/Core/Private/Android/AndroidProcess.cpp
--- a/Engine/Source/Runtime/Core/Private/Android/AndroidProcess.cpp
+++ b/Engine/Source/Runtime/Core/Private/Android/AndroidProcess.cpp
@@ -7,11 +7,22 @@
 #include "CorePrivate.h"
 #include "AndroidPlatformRunnableThread.h"
 
+#include <sys/syscall.h>
+#include <pthread.h>
+
+
 const TCHAR* FAndroidPlatformProcess::ComputerName()
 {
 	return TEXT("Android Device"); 
 }
 
+void FAndroidPlatformProcess::SetThreadAffinityMask( uint64 InAffinityMask )
+{
+	pid_t ThreadId = gettid();
+	int AffinityMask = (int)InAffinityMask;
+	syscall(__NR_sched_setaffinity, ThreadId, sizeof(AffinityMask), &AffinityMask);
+}
+
 const TCHAR* FAndroidPlatformProcess::BaseDir()
 {
 	return TEXT("");
diff --git a/Engine/Source/Runtime/Core/Private/GenericPlatform/GenericPlatformProcess.cpp b/Engine/Source/Runtime/Core/Private/GenericPlatform/GenericPlatformProcess.cpp
--- a/Engine/Source/Runtime/Core/Private/GenericPlatform/GenericPlatformProcess.cpp
+++ b/Engine/Source/Runtime/Core/Private/GenericPlatform/GenericPlatformProcess.cpp
@@ -40,6 +40,11 @@ uint32 FGenericPlatformProcess::GetCurrentProcessId()
 	return 0;
 }
 
+void FGenericPlatformProcess::SetThreadAffinityMask( uint64 AffinityMask )
+{
+	// Not implemented cross-platform. Each platform may or may not choose to implement this.
+}
+
 const TCHAR* FGenericPlatformProcess::UserDir()
 {
 	// default to the root directory
diff --git a/Engine/Source/Runtime/Core/Private/HAL/AffinityManager.cpp b/Engine/Source/Runtime/Core/Private/HAL/AffinityManager.cpp
--- a/Engine/Source/Runtime/Core/Private/HAL/AffinityManager.cpp
+++ b/Engine/Source/Runtime/Core/Private/HAL/AffinityManager.cpp
@@ -2,6 +2,8 @@
 
 #include "CorePrivate.h"
 
+//@TODO: All this should be moved into platform-specific FPlatformProcess. All platform #ifdef can then be removed.
+
 #define MAXAFFINITYMANAGERTHREADS (32)
 
 #define		MAKEAFFINITYMASK1(x) (1<<x)
@@ -54,7 +56,7 @@ uint64 AffinityManagerGetAffinity( const TCHAR* ThreadName )
 {
 	if( !GIsEditor )
 	{
-#if PLATFORM_XBOXONE
+#if PLATFORM_XBOXONE || PLATFORM_ANDROID
 		for( int i = 0; i < MAXAFFINITYMANAGERTHREADS; i++ )
 		{
 			if( gAffinityData[i].name )
diff --git a/Engine/Source/Runtime/Core/Private/Windows/WindowsPlatformProcess.cpp b/Engine/Source/Runtime/Core/Private/Windows/WindowsPlatformProcess.cpp
--- a/Engine/Source/Runtime/Core/Private/Windows/WindowsPlatformProcess.cpp
+++ b/Engine/Source/Runtime/Core/Private/Windows/WindowsPlatformProcess.cpp
@@ -367,6 +367,11 @@ uint32 FWindowsPlatformProcess::GetCurrentProcessId()
 	return ::GetCurrentProcessId();
 }
 
+void FWindowsPlatformProcess::SetThreadAffinityMask( uint64 AffinityMask )
+{
+	::SetThreadAffinityMask(::GetCurrentThread(), (DWORD_PTR)AffinityMask);
+}
+
 bool FWindowsPlatformProcess::GetProcReturnCode( FProcHandle & ProcHandle, int32* ReturnCode )
 {
 	return ::GetExitCodeProcess( ProcHandle.Get(), (::DWORD *)ReturnCode ) && *((uint32*)ReturnCode) != ProcessConstants::WIN_STILL_ACTIVE;
diff --git a/Engine/Source/Runtime/Core/Public/Android/AndroidProcess.h b/Engine/Source/Runtime/Core/Public/Android/AndroidProcess.h
--- a/Engine/Source/Runtime/Core/Public/Android/AndroidProcess.h
+++ b/Engine/Source/Runtime/Core/Public/Android/AndroidProcess.h
@@ -28,6 +28,7 @@ public:
 struct CORE_API FAndroidPlatformProcess : public FGenericPlatformProcess
 {
 	static const TCHAR* ComputerName();
+	static void SetThreadAffinityMask( uint64 AffinityMask );
 	static const TCHAR* BaseDir();
 	static const TCHAR* ExecutableName(bool bRemoveExtension = true);
 	static class FRunnableThread* CreateRunnableThread();
diff --git a/Engine/Source/Runtime/Core/Public/GenericPlatform/GenericPlatformProcess.h b/Engine/Source/Runtime/Core/Public/GenericPlatform/GenericPlatformProcess.h
--- a/Engine/Source/Runtime/Core/Public/GenericPlatform/GenericPlatformProcess.h
+++ b/Engine/Source/Runtime/Core/Public/GenericPlatform/GenericPlatformProcess.h
@@ -194,10 +194,13 @@ struct CORE_API FGenericPlatformProcess
 
 	}
 
-	/** Retrieves the ProcessId of this process
-	* @return the ProcessId of this process
-	*/
+	/**
+	 * Retrieves the ProcessId of this process
+	 * @return the ProcessId of this process
+	 */
 	static uint32 GetCurrentProcessId();
+	/** Sets the CPU affinity mask for the current thread. */
+	static void SetThreadAffinityMask( uint64 AffinityMask );
 	/** Get startup directory.  NOTE: Only one return value is valid at a time! **/
 	static const TCHAR* BaseDir();
 	/** Get user directory.  NOTE: Only one return value is valid at a time! **/
diff --git a/Engine/Source/Runtime/Core/Public/Windows/WindowsPlatformProcess.h b/Engine/Source/Runtime/Core/Public/Windows/WindowsPlatformProcess.h
--- a/Engine/Source/Runtime/Core/Public/Windows/WindowsPlatformProcess.h
+++ b/Engine/Source/Runtime/Core/Public/Windows/WindowsPlatformProcess.h
@@ -75,6 +75,7 @@ struct CORE_API FWindowsPlatformProcess : public FGenericPlatformProcess
 	static void PopDllDirectory(const TCHAR* Directory);
 	static void CleanFileCache();
 	static uint32 GetCurrentProcessId();
+	static void SetThreadAffinityMask( uint64 AffinityMask );
 	static const TCHAR* BaseDir();
 	static const TCHAR* UserDir();
 	static const TCHAR* UserSettingsDir();
diff --git a/Engine/Source/Runtime/Launch/Private/Android/LaunchAndroid.cpp b/Engine/Source/Runtime/Launch/Private/Android/LaunchAndroid.cpp
--- a/Engine/Source/Runtime/Launch/Private/Android/LaunchAndroid.cpp
+++ b/Engine/Source/Runtime/Launch/Private/Android/LaunchAndroid.cpp
@@ -243,6 +243,9 @@ static void* AndroidEventThreadWorker( void* param )
 {
 	struct android_app* state = (struct android_app*)param;
 
+	uint64 GameThreadAffinity = AffinityManagerGetAffinity( TEXT("MainGame"));
+	FPlatformProcess::SetThreadAffinityMask( GameThreadAffinity );
+
 	FPlatformMisc::LowLevelOutputDebugString(L"Entering event processing thread engine entry point");
 
 	ALooper* looper = ALooper_prepare(ALOOPER_PREPARE_ALLOW_NON_CALLBACKS);
diff --git a/Engine/Source/Runtime/Launch/Private/LaunchEngineLoop.cpp b/Engine/Source/Runtime/Launch/Private/LaunchEngineLoop.cpp
--- a/Engine/Source/Runtime/Launch/Private/LaunchEngineLoop.cpp
+++ b/Engine/Source/Runtime/Launch/Private/LaunchEngineLoop.cpp
@@ -612,10 +612,15 @@ int32 FEngineLoop::PreInit( const TCHAR* CmdLine )
 		UE_LOG(LogInit, Log, TEXT("Process affinity %ld System affinity %ld."), dwProcessAffinity , dwSystemAffinity);
 	}
 	uint64 GameThreadAffinity = AffinityManagerGetAffinity( TEXT("MainGame"));
-	::SetThreadAffinityMask(::GetCurrentThread(), (DWORD_PTR)GameThreadAffinity);
+	FPlatformProcess::SetThreadAffinityMask( GameThreadAffinity );
 	UE_LOG(LogInit, Log, TEXT("Runnable thread Main Thread is on Process %d."), static_cast<uint32>(::GetCurrentProcessorNumber()) );	
 #endif // PLATFORM_XBOXONE
 
+#if PLATFORM_ANDROID
+	uint64 GameThreadAffinity = AffinityManagerGetAffinity( TEXT("MainGame"));
+	FPlatformProcess::SetThreadAffinityMask( GameThreadAffinity );
+#endif
+
 	// Figure out whether we're the editor, ucc or the game.
 	const SIZE_T CommandLineSize = FCString::Strlen(CmdLine)+1;
 	TCHAR* CommandLineCopy			= new TCHAR[ CommandLineSize ];
diff --git a/Engine/Source/Runtime/RenderCore/Private/RenderingThread.cpp b/Engine/Source/Runtime/RenderCore/Private/RenderingThread.cpp
--- a/Engine/Source/Runtime/RenderCore/Private/RenderingThread.cpp
+++ b/Engine/Source/Runtime/RenderCore/Private/RenderingThread.cpp
@@ -23,7 +23,7 @@ RENDERCORE_API bool GUseThreadedRendering = false;
 #if !(UE_BUILD_SHIPPING || UE_BUILD_TEST)
 	RENDERCORE_API bool GMainThreadBlockedOnRenderThread = false;
 #endif // #if !(UE_BUILD_SHIPPING || UE_BUILD_TEST)
-
+	
 static FRunnable* GRenderingThreadRunnable = NULL;
 
 /** If the rendering thread has been terminated by an unhandled exception, this contains the error message. */
@@ -312,6 +312,12 @@ public:
 	{
 		GRenderThreadId = FPlatformTLS::GetCurrentThreadId();
 
+#if PLATFORM_ANDROID
+		//@TODO: This should be using the affinitymask in FRunnableThread::Create() and tested on all platforms.
+		uint64 RenderThreadAffinityMask = AffinityManagerGetAffinity( TEXT("RenderingThread 0"));
+		FPlatformProcess::SetThreadAffinityMask( RenderThreadAffinityMask );
+#endif
+
 #if PLATFORM_WINDOWS
 		if ( !FPlatformMisc::IsDebuggerPresent() || GAlwaysReportCrash )
 		{
