#UE4 - removed unnecessary string manipulation

diff --git a/Engine/Source/Runtime/Launch/Private/LaunchEngineLoop.cpp b/Engine/Source/Runtime/Launch/Private/LaunchEngineLoop.cpp
--- a/Engine/Source/Runtime/Launch/Private/LaunchEngineLoop.cpp
+++ b/Engine/Source/Runtime/Launch/Private/LaunchEngineLoop.cpp
@@ -665,7 +665,7 @@ int32 FEngineLoop::PreInit( const TCHAR* CmdLine )
 
 	const bool bFirstTokenIsGameName = (FApp::HasGameName() && Token == GGameName);
 	const bool bFirstTokenIsGameProjectFilePath = (FPaths::IsProjectFilePathSet() && Token.Replace(TEXT("\\"), TEXT("/")) == FPaths::GetProjectFilePath());
-	const bool bFirstTokenIsGameProjectFileShortName = (FPaths::IsProjectFilePathSet() && Token.Replace(TEXT("\\"), TEXT("/")) == FPaths::GetCleanFilename(FPaths::GetProjectFilePath()));
+	const bool bFirstTokenIsGameProjectFileShortName = (FPaths::IsProjectFilePathSet() && Token == FPaths::GetCleanFilename(FPaths::GetProjectFilePath()));
 
 	if (bFirstTokenIsGameName || bFirstTokenIsGameProjectFilePath || bFirstTokenIsGameProjectFileShortName)
 	{
