- Fixed TTP:329773  Android volume buttons don't work.
- We were telling the OS that we were handling all key input.  Changed to only say we've handled mapped buttons.

diff --git a/Engine/Source/Runtime/Core/Private/Android/AndroidMisc.cpp b/Engine/Source/Runtime/Core/Private/Android/AndroidMisc.cpp
--- a/Engine/Source/Runtime/Core/Private/Android/AndroidMisc.cpp
+++ b/Engine/Source/Runtime/Core/Private/Android/AndroidMisc.cpp
@@ -237,11 +237,11 @@ bool FAndroidMisc::GetUseVirtualJoysticks()
 
 uint32 FAndroidMisc::GetKeyMap( uint16* KeyCodes, FString* KeyNames, uint32 MaxMappings )
 {
-#define ADDKEYMAP(KeyCode, KeyName)		if (NumMappings<MaxMappings) { KeyCodes[NumMappings]=KeyCode; KeyNames[NumMappings]=KeyName; ++NumMappings; };
+#define ADDKEYMAP(KeyCode, KeyName)		if (NumMappings<MaxMappings) { KeyCodes[NumMappings]=KeyCode; if(KeyNames) { KeyNames[NumMappings]=KeyName; } ++NumMappings; };
 
 	uint32 NumMappings = 0;
 
-	if ( KeyCodes && KeyNames && (MaxMappings > 0) )
+	if ( KeyCodes && (MaxMappings > 0) )
 	{
 		ADDKEYMAP( AKEYCODE_0, TEXT("Zero") );
 		ADDKEYMAP( AKEYCODE_1, TEXT("One") );
diff --git a/Engine/Source/Runtime/Launch/Private/Android/LaunchAndroid.cpp b/Engine/Source/Runtime/Launch/Private/Android/LaunchAndroid.cpp
--- a/Engine/Source/Runtime/Launch/Private/Android/LaunchAndroid.cpp
+++ b/Engine/Source/Runtime/Launch/Private/Android/LaunchAndroid.cpp
@@ -39,6 +39,9 @@ static const int32_t AxisList[] =
     //AMOTION_EVENT_AXIS_HAT_Y,
 };
 
+// map of all supported keycodes
+static TSet<uint16> MappedKeyCodes;
+
 // -nostdlib means no crtbegin_so.o, so we have to provide our own __dso_handle and atexit()
 extern "C"
 {
@@ -192,6 +195,16 @@ int32 AndroidMain(struct android_app* state)
 		}
 	}
 
+	// setup key filtering
+	static const uint32 MAX_KEY_MAPPINGS(256);
+	uint16 KeyCodes[MAX_KEY_MAPPINGS];
+	uint32 NumKeyCodes = FPlatformMisc::GetKeyMap(KeyCodes, nullptr, MAX_KEY_MAPPINGS);
+
+	for (int i = 0; i < NumKeyCodes; ++i)
+	{
+		MappedKeyCodes.Add(KeyCodes[i]);
+	}
+
 	// read the command line file
 	InitCommandLine();
 	FPlatformMisc::LowLevelOutputDebugStringf(TEXT("Final commandline: %s\n"), FCommandLine::Get());
@@ -372,7 +385,7 @@ static int32_t HandleInputCB(struct android_app* app, AInputEvent* event)
 			ANativeWindow* Window = (ANativeWindow*)FPlatformMisc::GetHardwareWindow();
 			if (!Window)
 			{
-				return 1;
+				return 0;
 			}
 
 			int32_t Width = 0 ;
@@ -457,6 +470,12 @@ static int32_t HandleInputCB(struct android_app* app, AInputEvent* event)
 		{
 			FPlatformMisc::LowLevelOutputDebugStringf(L"Received key event: %d", keyCode);
 
+			// only handle mapped key codes
+			if (!MappedKeyCodes.Contains(keyCode))
+			{
+				return 0;
+			}
+
 			FDeferredAndroidMessage Message;
 
 			Message.messageType = AKeyEvent_getAction(event) == AKEY_EVENT_ACTION_UP ? MessageType_KeyUp : MessageType_KeyDown; 
