#UE4 - Initialize std out device as soon as possible so that UAT can grab all the logs.

diff --git a/Engine/Source/Runtime/Launch/Private/LaunchEngineLoop.cpp b/Engine/Source/Runtime/Launch/Private/LaunchEngineLoop.cpp
--- a/Engine/Source/Runtime/Launch/Private/LaunchEngineLoop.cpp
+++ b/Engine/Source/Runtime/Launch/Private/LaunchEngineLoop.cpp
@@ -109,6 +109,18 @@ private:
 static TScopedPointer<FOutputDeviceConsole>	GScopedLogConsole;
 static TScopedPointer<FOutputDeviceStdOutput> GScopedStdOut;
 
+/**
+ * Initializes std out device and adds it to GLog
+ **/
+void InitializeStdOutDevice()
+{
+	// Check if something is trying to initialize std out device twice.
+	check(!GScopedStdOut.IsValid());
+
+	GScopedStdOut = new FOutputDeviceStdOutput();
+	GLog->AddOutputDevice(GScopedStdOut.GetOwnedPointer());
+}
+
 bool ParseGameProjectFromCommandLine(const TCHAR* InCmdLine, FString& OutProjectFilePath, FString& OutGameName)
 {
 	const TCHAR *CmdLine = InCmdLine;
@@ -544,6 +556,12 @@ int32 FEngineLoop::PreInit( const TCHAR* CmdLine )
 	// also realize that command lines can be pulled from the network at a slightly later time.
 	FCommandLine::Set(CmdLine); 
 
+	// Initialize std out device as early as possible if requested in the command line
+	if (FParse::Param(FCommandLine::Get(), TEXT("stdout")))
+	{
+		InitializeStdOutDevice();
+	}
+
 	// Set GGameName, based on the command line
 	if (LaunchSetGameName(CmdLine) == false)
 	{
@@ -1007,6 +1025,13 @@ int32 FEngineLoop::PreInit( const TCHAR* CmdLine )
 		bIsSeekFreeDedicatedServer = FPlatformProperties::RequiresCookedData();
 	}
 
+	// If std out device hasn't been initialzied yet (there was no -stdout param in the command line) and
+	// we meet all the criteria, initialize it now.
+	if (!GScopedStdOut.IsValid() && !bHasEditorToken && !bIsRegularClient && !IsRunningDedicatedServer())
+	{
+		InitializeStdOutDevice();
+	}
+
 	// Initialize the RHI.
 	RHIInit(bHasEditorToken);
 
@@ -1193,12 +1218,6 @@ int32 FEngineLoop::PreInit( const TCHAR* CmdLine )
 	MarkObjectsToDisregardForGC(); 
 	GUObjectArray.CloseDisregardForGC();
 
-	if (FParse::Param(FCommandLine::Get(), TEXT("stdout")) || (!bHasEditorToken && !bIsRegularClient && !IsRunningDedicatedServer()))
-	{
-		GScopedStdOut = new FOutputDeviceStdOutput(); 
-		GLog->AddOutputDevice( GScopedStdOut.GetOwnedPointer() );
-	}
-
 #if WITH_ENGINE
 	SetIsServerForOnlineSubsystemsDelegate(FQueryIsRunningServer::CreateStatic(&IsServerDelegateForOSS));
 #endif
