UE4: Linux dedicated servers use no audio device, so we need not depend on the OpenAL lib.

diff --git a/Engine/Source/Runtime/Launch/Launch.Build.cs b/Engine/Source/Runtime/Launch/Launch.Build.cs
--- a/Engine/Source/Runtime/Launch/Launch.Build.cs
+++ b/Engine/Source/Runtime/Launch/Launch.Build.cs
@@ -166,7 +166,7 @@ public class Launch : ModuleRules
 			DynamicallyLoadedModuleNames.Add("OpenGLDrv");
 		}
 
-		if ((Target.Platform == UnrealTargetPlatform.Linux))
+        if ((Target.Platform == UnrealTargetPlatform.Linux) && (Target.Type != TargetRules.TargetType.Server))
 		{
 			PrivateDependencyModuleNames.Add("ALAudio");
 		}
