#UE4:
* fixed Optimus GPU detection (using faster non integrated one) when running editor

diff --git a/Engine/Source/Runtime/Launch/Private/Windows/LaunchWindows.cpp b/Engine/Source/Runtime/Launch/Private/Windows/LaunchWindows.cpp
--- a/Engine/Source/Runtime/Launch/Private/Windows/LaunchWindows.cpp
+++ b/Engine/Source/Runtime/Launch/Private/Windows/LaunchWindows.cpp
@@ -12,6 +12,11 @@ DEFINE_LOG_CATEGORY_STATIC(LogLaunchWindows, Log, All);
 extern int32 GuardedMain( const TCHAR* CmdLine, HINSTANCE hInInstance, HINSTANCE hPrevInstance, int32 nCmdShow );
 extern void LaunchStaticShutdownAfterError();
 
+// http://developer.download.nvidia.com/devzone/devcenter/gamegraphics/files/OptimusRenderingPolicies.pdf
+// The following line is to favor the high performance NVIDIA GPU if there are multiple GPUs
+// Has to be .exe module to be correctly detected.
+extern "C" { _declspec(dllexport) uint32 NvOptimusEnablement = 0x00000001; }
+
 /**
  * Maintain a named mutex to detect whether we are the first instance of this game
  */
diff --git a/Engine/Source/Runtime/Windows/D3D11RHI/Private/D3D11RHI.cpp b/Engine/Source/Runtime/Windows/D3D11RHI/Private/D3D11RHI.cpp
--- a/Engine/Source/Runtime/Windows/D3D11RHI/Private/D3D11RHI.cpp
+++ b/Engine/Source/Runtime/Windows/D3D11RHI/Private/D3D11RHI.cpp
@@ -28,6 +28,7 @@ extern void UniformBufferBeginFrame();
 
 // http://developer.download.nvidia.com/devzone/devcenter/gamegraphics/files/OptimusRenderingPolicies.pdf
 // The following line is to favor the high performance NVIDIA GPU if there are multiple GPUs
+// Has to be .exe module to be correctly detected.
 extern "C" { _declspec(dllexport) uint32 NvOptimusEnablement = 0x00000001; }
 
 void FD3D11DynamicRHI::RHIBeginFrame()
