Engine: L10N: Launch engine loop will now call the EndTextLocalizationInit function along with the Begin version in standalone applications.

diff --git a/Engine/Source/Runtime/Launch/Private/LaunchEngineLoop.cpp b/Engine/Source/Runtime/Launch/Private/LaunchEngineLoop.cpp
--- a/Engine/Source/Runtime/Launch/Private/LaunchEngineLoop.cpp
+++ b/Engine/Source/Runtime/Launch/Private/LaunchEngineLoop.cpp
@@ -1491,8 +1491,11 @@ int32 FEngineLoop::PreInit( const TCHAR* CmdLine )
 	}
 #endif
 
+
+#else
+	EndInitTextLocalization();
 #endif
-	
+
 	//run automation smoke tests now that everything is setup to run
 	FAutomationTestFramework::GetInstance().RunSmokeTests();
 	
