- Make four finger touch more reliable.  Fixes TTP:331266

diff --git a/Engine/Build/Android/Java/src/com/epicgames/ue4/GameActivity.java b/Engine/Build/Android/Java/src/com/epicgames/ue4/GameActivity.java
--- a/Engine/Build/Android/Java/src/com/epicgames/ue4/GameActivity.java
+++ b/Engine/Build/Android/Java/src/com/epicgames/ue4/GameActivity.java
@@ -99,6 +99,12 @@ public class GameActivity extends NativeActivity
 	// Called from event thread in NativeActivity	
 	public void AndroidThunkJava_ShowConsoleWindow(String Formats)
 	{
+		if(alert.isShowing() == true)
+		{
+			Log.debug("Console already showing.");
+			return;
+		}
+
 		alert.setMessage("[Availble texture formats: " + Formats + "]");
 		_activity.runOnUiThread(new Runnable()
 		{
diff --git a/Engine/Source/Runtime/Launch/Private/Android/LaunchAndroid.cpp b/Engine/Source/Runtime/Launch/Private/Android/LaunchAndroid.cpp
--- a/Engine/Source/Runtime/Launch/Private/Android/LaunchAndroid.cpp
+++ b/Engine/Source/Runtime/Launch/Private/Android/LaunchAndroid.cpp
@@ -305,8 +305,8 @@ static void AndroidProcessEvents(struct android_app* state)
 		// process this event
 		if (source)
 			source->process(state, source);
+		}
 	}
-}
 
 pthread_t G_AndroidEventThread;
 
@@ -445,7 +445,7 @@ static int32_t HandleInputCB(struct android_app* app, AInputEvent* event)
 			FAndroidInputInterface::QueueTouchInput(TouchesArray);
 
 #if !UE_BUILD_SHIPPING
-			if(pointerCount >= 4)
+			if ((pointerCount >= 4) && (type == TouchBegan))
 			{
 				GShowConsoleWindowNextTick = true;
 			}
