#ue4
- Added a method to show the platform-specific Achievements screen on Android

diff --git a/Engine/Build/Android/Java/src/com/epicgames/ue4/GameActivity.java b/Engine/Build/Android/Java/src/com/epicgames/ue4/GameActivity.java
--- a/Engine/Build/Android/Java/src/com/epicgames/ue4/GameActivity.java
+++ b/Engine/Build/Android/Java/src/com/epicgames/ue4/GameActivity.java
@@ -269,6 +269,15 @@ public class GameActivity extends NativeActivity implements GoogleApiClient.Conn
 		}
 	}
 
+	public void AndroidThunkJava_ShowAchievements()
+	{
+		Log.debug("In AndroidThunkJava_ShowAchievements");
+		if(googleClient.isConnected())
+		{
+			startActivityForResult(Games.Achievements.getAchievementsIntent(googleClient), 0);
+		}
+	}
+
 	public void AndroidThunkJava_WriteLeaderboardValue(String LeaderboardID, long Value)
 	{
 		Log.debug("In AndroidThunkJava_WriteLeaderboardValue, ID is " + LeaderboardID + ", value is " + Value);
diff --git a/Engine/Source/Runtime/Launch/Private/Android/AndroidJNI.cpp b/Engine/Source/Runtime/Launch/Private/Android/AndroidJNI.cpp
--- a/Engine/Source/Runtime/Launch/Private/Android/AndroidJNI.cpp
+++ b/Engine/Source/Runtime/Launch/Private/Android/AndroidJNI.cpp
@@ -95,6 +95,7 @@ jclass JDef_GameActivity::ClassID;
 jmethodID JDef_GameActivity::AndroidThunkJava_ShowConsoleWindow;
 jmethodID JDef_GameActivity::AndroidThunkJava_LaunchURL;
 jmethodID JDef_GameActivity::AndroidThunkJava_ShowLeaderboard;
+jmethodID JDef_GameActivity::AndroidThunkJava_ShowAchievements;
 jmethodID JDef_GameActivity::AndroidThunkJava_WriteLeaderboardValue;
 jmethodID JDef_GameActivity::AndroidThunkJava_GooglePlayConnect;
 
@@ -174,6 +175,14 @@ void AndroidThunkCpp_ShowLeaderboard(const FString& CategoryName)
 	}
 }
 
+void AndroidThunkCpp_ShowAchievements()
+{
+ 	if (JNIEnv* Env = GetJavaEnv())
+ 	{
+		Env->CallVoidMethod(GJavaGlobalThis, JDef_GameActivity::AndroidThunkJava_ShowAchievements);
+ 	}
+}
+
 void AndroidThunkCpp_WriteLeaderboardValue(const FString& LeaderboardName, int64_t Value)
 {
 	if (JNIEnv* Env = GetJavaEnv())
@@ -209,6 +218,7 @@ JNIEXPORT jint JNI_OnLoad(JavaVM* InJavaVM, void* InReserved)
 	JDef_GameActivity::AndroidThunkJava_ShowConsoleWindow = env->GetMethodID(JDef_GameActivity::ClassID, "AndroidThunkJava_ShowConsoleWindow", "(Ljava/lang/String;)V");
 	JDef_GameActivity::AndroidThunkJava_LaunchURL = env->GetMethodID(JDef_GameActivity::ClassID, "AndroidThunkJava_LaunchURL", "(Ljava/lang/String;)V");
 	JDef_GameActivity::AndroidThunkJava_ShowLeaderboard = env->GetMethodID(JDef_GameActivity::ClassID, "AndroidThunkJava_ShowLeaderboard", "(Ljava/lang/String;)V");
+	JDef_GameActivity::AndroidThunkJava_ShowAchievements = env->GetMethodID(JDef_GameActivity::ClassID, "AndroidThunkJava_ShowAchievements", "()V");
 	JDef_GameActivity::AndroidThunkJava_WriteLeaderboardValue = env->GetMethodID(JDef_GameActivity::ClassID, "AndroidThunkJava_WriteLeaderboardValue", "(Ljava/lang/String;J)V");
 	JDef_GameActivity::AndroidThunkJava_GooglePlayConnect = env->GetMethodID(JDef_GameActivity::ClassID, "AndroidThunkJava_GooglePlayConnect", "()V");
 
diff --git a/Engine/Source/Runtime/Launch/Private/Android/AndroidJNI.h b/Engine/Source/Runtime/Launch/Private/Android/AndroidJNI.h
--- a/Engine/Source/Runtime/Launch/Private/Android/AndroidJNI.h
+++ b/Engine/Source/Runtime/Launch/Private/Android/AndroidJNI.h
@@ -22,6 +22,7 @@ public:
 	static jmethodID AndroidThunkJava_ShowConsoleWindow;
 	static jmethodID AndroidThunkJava_LaunchURL;
 	static jmethodID AndroidThunkJava_ShowLeaderboard;
+	static jmethodID AndroidThunkJava_ShowAchievements;
 	static jmethodID AndroidThunkJava_WriteLeaderboardValue;
 	static jmethodID AndroidThunkJava_GooglePlayConnect;
 };
