Fix for TTP #333174 CRASH: MAC: The editor crashes immediately when trying to launch ShooterGame 4.1

Don't load MainFrame module to display about window, but rather only display it if the module is loaded. Solves a rare crash where user tried to display about window too early.

diff --git a/Engine/Source/Runtime/Launch/Private/Mac/LaunchMac.cpp b/Engine/Source/Runtime/Launch/Private/Mac/LaunchMac.cpp
--- a/Engine/Source/Runtime/Launch/Private/Mac/LaunchMac.cpp
+++ b/Engine/Source/Runtime/Launch/Private/Mac/LaunchMac.cpp
@@ -141,7 +141,10 @@ void EngineCrashHandler(const FGenericCrashContext & GenericContext)
 - (IBAction)OnShowAboutWindow:(id)Sender
 {
 #if WITH_EDITOR
-	FModuleManager::LoadModuleChecked<IMainFrameModule>(TEXT("MainFrame")).ShowAboutWindow();
+	if (FModuleManager::Get().IsModuleLoaded(TEXT("MainFrame")))
+	{
+		FModuleManager::GetModuleChecked<IMainFrameModule>(TEXT("MainFrame")).ShowAboutWindow();
+	}
 #else
 	[NSApp orderFrontStandardAboutPanel: Sender];
 #endif
