Fixed volume buttons on Android.

- this was a regression caused by fixing gamepad button support.  Apparently the volume buttons are coming through as gamepad buttons, so weren't hitting the keycode filtering code, causing the OS to think we were handling them.
- Special cased for now, although we should revisit gamepad handling in general.

diff --git a/Engine/Source/Runtime/Launch/Private/Android/LaunchAndroid.cpp b/Engine/Source/Runtime/Launch/Private/Android/LaunchAndroid.cpp
--- a/Engine/Source/Runtime/Launch/Private/Android/LaunchAndroid.cpp
+++ b/Engine/Source/Runtime/Launch/Private/Android/LaunchAndroid.cpp
@@ -42,6 +42,16 @@ static const int32_t AxisList[] =
 // map of all supported keycodes
 static TSet<uint16> MappedKeyCodes;
 
+// List of gamepad keycodes to ignore
+static const uint16 IgnoredGamepadKeyCodesList[] =
+{
+	AKEYCODE_VOLUME_UP,
+	AKEYCODE_VOLUME_DOWN
+};
+
+// map of gamepad keycodes that should be ignored
+static TSet<uint16> IgnoredGamepadKeyCodes;
+
 // -nostdlib means no crtbegin_so.o, so we have to provide our own __dso_handle and atexit()
 extern "C"
 {
@@ -208,6 +218,12 @@ int32 AndroidMain(struct android_app* state)
 		MappedKeyCodes.Add(KeyCodes[i]);
 	}
 
+	const int IgnoredGamepadKeyCodeCount = sizeof(IgnoredGamepadKeyCodesList)/sizeof(uint16);
+	for (int i = 0; i < IgnoredGamepadKeyCodeCount; ++i)
+	{
+		IgnoredGamepadKeyCodes.Add(IgnoredGamepadKeyCodesList[i]);
+	}
+
 	// read the command line file
 	InitCommandLine();
 	FPlatformMisc::LowLevelOutputDebugStringf(TEXT("Final commandline: %s\n"), FCommandLine::Get());
@@ -471,9 +487,16 @@ static int32_t HandleInputCB(struct android_app* app, AInputEvent* event)
 	{
 		int keyCode = AKeyEvent_getKeyCode(event);
 
+		FPlatformMisc::LowLevelOutputDebugStringf(L"Received keycode: %d", keyCode);
+
 		//Trap Joystick events first, with fallthrough if there is no joystick support
 		if (((AInputEvent_getSource(event) & (AINPUT_SOURCE_GAMEPAD | AINPUT_SOURCE_DPAD)) != 0) && (GetAxes != NULL))
 		{
+			if (IgnoredGamepadKeyCodes.Contains(keyCode))
+			{
+				return 0;
+			}
+
 			int device = AInputEvent_getDeviceId(event);
 			bool down = AKeyEvent_getAction(event) != AKEY_EVENT_ACTION_UP;
 			FAndroidInputInterface::JoystickButtonEvent( device, keyCode, down);
