Allow -waitfordebugger in iOS Release build to debug packaged startup crashes

diff --git a/Engine/Source/Runtime/Launch/Private/IOS/LaunchIOS.cpp b/Engine/Source/Runtime/Launch/Private/IOS/LaunchIOS.cpp
--- a/Engine/Source/Runtime/Launch/Private/IOS/LaunchIOS.cpp
+++ b/Engine/Source/Runtime/Launch/Private/IOS/LaunchIOS.cpp
@@ -179,7 +179,7 @@ int main(int argc, char *argv[])
 
 	FIOSCommandLineHelper::InitCommandArgs(FString());
 	
-#if UE_BUILD_DEBUG
+#if !UE_BUILD_SHIPPING
     if (FParse::Param(FCommandLine::Get(), TEXT("WaitForDebugger")))
     {
         while(!FPlatformMisc::IsDebuggerPresent())
