Update version number in main to 4.2.

diff --git a/Engine/Source/Runtime/Launch/Resources/Version.h b/Engine/Source/Runtime/Launch/Resources/Version.h
--- a/Engine/Source/Runtime/Launch/Resources/Version.h
+++ b/Engine/Source/Runtime/Launch/Resources/Version.h
@@ -3,7 +3,7 @@
 #pragma once
 
 #define ENGINE_MAJOR_VERSION	4
-#define ENGINE_MINOR_VERSION	1
+#define ENGINE_MINOR_VERSION	2
 #define ENGINE_PATCH_VERSION	0
 
 #define ENGINE_VERSION 0
