Changes to the way we increase rlimits on Linux.

- Do not adjust number of open files upfront, and don't use "-pak" switch for that logic (TTP #331889).
- Make adjusting core size optional to avoid problems like JoshM ran into when running servers from a space-limited laptop.

diff --git a/Engine/Source/Runtime/Launch/Private/Linux/LaunchLinux.cpp b/Engine/Source/Runtime/Launch/Private/Linux/LaunchLinux.cpp
--- a/Engine/Source/Runtime/Launch/Private/Linux/LaunchLinux.cpp
+++ b/Engine/Source/Runtime/Launch/Private/Linux/LaunchLinux.cpp
@@ -90,31 +90,38 @@ static bool IncreaseLimit(int Resource, rlim_t DesiredLimit)
 }
 
 /**
- * Increases limit on 
- *  - number of open files to be no less than desired, so we can handle loading un-pak-ed builds
- *  - size of core file, so core gets dumped and we can debug crashed builds
+ * Expects GSavedCommandLine to be set up. Increases limit on 
+ *  - number of open files to be no less than desired (if specified on command line, otherwise left alone)
+ *  - size of core file, so core gets dumped and we can debug crashed builds (unless overriden with -nocore)
  *
- * @param DesiredLimit - max number of open files desired.
  */
-static bool IncreasePerProcessLimits(int32 DesiredNumberOfFiles)
+static bool IncreasePerProcessLimits()
 {
+	// honor the parameter if given, but don't change limits if not
+	int32 FileHandlesToReserve = -1;
+	if (FParse::Value(*GSavedCommandLine, TEXT("numopenfiles="), FileHandlesToReserve) && FileHandlesToReserve > 0)
+	{
 #if !UE_BUILD_SHIPPING
-	printf("Increasing per-process limit of open file handles to %d\n", DesiredNumberOfFiles);
+		printf("Increasing per-process limit of open file handles to %d\n", FileHandlesToReserve);
 #endif // !UE_BUILD_SHIPPING
-	if ( !IncreaseLimit(RLIMIT_NOFILE, DesiredNumberOfFiles) )
-	{
-		fprintf(stderr, "Could not adjust number of file handles, consider changing \"nofile\" in /etc/security/limits.conf and relogin.\n", errno, strerror(errno));
-		return false;
+		if (!IncreaseLimit(RLIMIT_NOFILE, FileHandlesToReserve))
+		{
+			fprintf(stderr, "Could not adjust number of file handles, consider changing \"nofile\" in /etc/security/limits.conf and relogin.\n", errno, strerror(errno));
+			return false;
+		}
 	}
 
 #if !UE_BUILD_SHIPPING
-	printf("Increasing per-process limit of core file size to infinity.\n");
-#endif // !UE_BUILD_SHIPPING
-	if ( !IncreaseLimit(RLIMIT_CORE, RLIM_INFINITY))
+	if (!FParse::Param(*GSavedCommandLine, TEXT("nocore")))
 	{
-		fprintf(stderr, "Could not adjust core file size, consider changing \"core\" in /etc/security/limits.conf and relogin.\n", errno, strerror(errno));
-		return false;
+		printf("Increasing per-process limit of core file size to infinity.\n");
+		if (!IncreaseLimit(RLIMIT_CORE, RLIM_INFINITY))
+		{
+			fprintf(stderr, "Could not adjust core file size, consider changing \"core\" in /etc/security/limits.conf and relogin.\n", errno, strerror(errno));
+			return false;
+		}
 	}
+#endif // !UE_BUILD_SHIPPING
 
 	return true;
 }
@@ -145,19 +152,7 @@ int main(int argc, char *argv[])
 	}
 #endif
 
-	int32 FileHandlesToReserve = 10000, OverrideFileHandlesToReserve;
-	if (FParse::Param( *GSavedCommandLine,TEXT("pak")))
-	{
-		FileHandlesToReserve = 256;	// we need significantly fewer file handles when opening paks 
-	}
-
-	// allow to override
-	if (FParse::Value( *GSavedCommandLine, TEXT("numopenfiles"), OverrideFileHandlesToReserve) && OverrideFileHandlesToReserve > 0)
-	{
-		FileHandlesToReserve = OverrideFileHandlesToReserve;
-	}
-
-	if (!IncreasePerProcessLimits( FileHandlesToReserve ))
+	if (!IncreasePerProcessLimits())
 	{
 		fprintf(stderr, "Could not set desired per-process limits, consider changing system limits.\n");
 		ErrorLevel = 1;
