Minor: inform the user about "-nocore" switch.

diff --git a/Engine/Source/Runtime/Launch/Private/Linux/LaunchLinux.cpp b/Engine/Source/Runtime/Launch/Private/Linux/LaunchLinux.cpp
--- a/Engine/Source/Runtime/Launch/Private/Linux/LaunchLinux.cpp
+++ b/Engine/Source/Runtime/Launch/Private/Linux/LaunchLinux.cpp
@@ -118,6 +118,7 @@ static bool IncreasePerProcessLimits()
 		if (!IncreaseLimit(RLIMIT_CORE, RLIM_INFINITY))
 		{
 			fprintf(stderr, "Could not adjust core file size, consider changing \"core\" in /etc/security/limits.conf and relogin.\n", errno, strerror(errno));
+			fprintf(stderr, "Alternatively, pass -nocore if you are unable or unwilling to do that.\n");
 			return false;
 		}
 	}
