UE4: Temporarily added some additional initialization logs to track down a crash while cooking.

diff --git a/Engine/Source/Runtime/Core/Private/Internationalization/TextLocalizationManager.cpp b/Engine/Source/Runtime/Core/Private/Internationalization/TextLocalizationManager.cpp
--- a/Engine/Source/Runtime/Core/Private/Internationalization/TextLocalizationManager.cpp
+++ b/Engine/Source/Runtime/Core/Private/Internationalization/TextLocalizationManager.cpp
@@ -226,6 +226,8 @@ void FTextLocalizationManager::LoadResources(const bool ShouldLoadEditor, const
 	}
 	CultureTracker.ReportCollisions();
 
+	UE_LOG(LogInit, Log, TEXT("Text localization reported Culture Tracker collisions."));
+
 	// Read base language localization resources.
 	FLocalizationEntryTracker& BaseLanguageTracker = LocalizationEntryTrackers[LocalizationEntryTrackers.Add(FLocalizationEntryTracker())];
 	for(int32 PathIndex = 0; PathIndex < LocalizationPaths.Num(); ++PathIndex)
@@ -241,7 +243,11 @@ void FTextLocalizationManager::LoadResources(const bool ShouldLoadEditor, const
 	}
 	BaseLanguageTracker.ReportCollisions();
 
+	UE_LOG(LogInit, Log, TEXT("Text localization reported Base Language Tracker collisions."));
+
 	UpdateLiveTable(LocalizationEntryTrackers);
+
+	UE_LOG(LogInit, Log, TEXT("Text localization Live Table updated."));
 }
 
 void FTextLocalizationManager::UpdateLiveTable(const TArray<FLocalizationEntryTracker>& LocalizationEntryTrackers, const bool FilterUpdatesByTableName)
diff --git a/Engine/Source/Runtime/Engine/Private/GlobalShader.cpp b/Engine/Source/Runtime/Engine/Private/GlobalShader.cpp
--- a/Engine/Source/Runtime/Engine/Private/GlobalShader.cpp
+++ b/Engine/Source/Runtime/Engine/Private/GlobalShader.cpp
@@ -334,14 +334,22 @@ TShaderMap<FGlobalShaderType>* GetGlobalShaderMap(EShaderPlatform Platform, bool
 	// If the global shader map hasn't been created yet, create it.
 	if(!GGlobalShaderMap[Platform])
 	{
+		UE_LOG(LogInit, Log, TEXT("Beginning GetGlobalShaderMap"));
+
 		// verify that all shader source files are intact
 		VerifyShaderSourceFiles();
 
+		UE_LOG(LogInit, Log, TEXT("Shader source files verified"));
+
 		// GetGlobalShaderMap is called the first time during startup in the main thread.
 		check(IsInGameThread());
 
+		UE_LOG(LogInit, Log, TEXT("Verified game thread"));
+
 		GGlobalShaderMap[Platform] = new TShaderMap<FGlobalShaderType>();
 
+		UE_LOG(LogInit, Log, TEXT("Created global shader map"));
+
 		bool bLoadedFromCacheFile = false;
 
 		// Try to load the global shaders from a local cache file if it exists
@@ -378,18 +386,24 @@ TShaderMap<FGlobalShaderType>* GetGlobalShaderMap(EShaderPlatform Platform, bool
 		// Uncooked platform
 		else
 		{
+			UE_LOG(LogInit, Log, TEXT("Preparing global shader map on uncooked platform"));
 			FGlobalShaderMapId ShaderMapId(Platform);
 
 			TArray<uint8> CachedData;
 			const FString DataKey = GetGlobalShaderMapKeyString(ShaderMapId, Platform);
 
+			UE_LOG(LogInit, Log, TEXT("GlobalShaderMapKeyString formed"));
+
 			// Find the shader map in the derived data cache
 			if (GetDerivedDataCacheRef().GetSynchronous(*DataKey, CachedData))
 			{
+				UE_LOG(LogInit, Log, TEXT("GlobalShaderMap DDC data fetched"));
 				FMemoryReader Ar(CachedData, true);
 
 				// Deserialize from the cached data
 				SerializeGlobalShaders(Ar, GGlobalShaderMap[Platform]);
+
+				UE_LOG(LogInit, Log, TEXT("GlobalShaderMap serialized"));
 			}
 		}
 
diff --git a/Engine/Source/Runtime/Launch/Private/LaunchEngineLoop.cpp b/Engine/Source/Runtime/Launch/Private/LaunchEngineLoop.cpp
--- a/Engine/Source/Runtime/Launch/Private/LaunchEngineLoop.cpp
+++ b/Engine/Source/Runtime/Launch/Private/LaunchEngineLoop.cpp
@@ -1109,8 +1109,12 @@ int32 FEngineLoop::PreInit( const TCHAR* CmdLine )
 
 	EndInitTextLocalization();
 
+	UE_LOG(LogInit, Log, TEXT("Text localization initialized."));
+
 	GStreamingManager = new FStreamingManagerCollection();
 
+	UE_LOG(LogInit, Log, TEXT("StreamingManager created."));
+
 	if (FPlatformProcess::SupportsMultithreading() && !IsRunningDedicatedServer() && (bIsRegularClient || bHasEditorToken))
 	{
 		FPlatformSplash::Show();
@@ -1122,15 +1126,18 @@ int32 FEngineLoop::PreInit( const TCHAR* CmdLine )
 		{
 			// Init platform application
 			FSlateApplication::Create();
+			UE_LOG(LogInit, Log, TEXT("Slate application created."));
 		}
 		else
 		{
 			FCoreStyle::ResetToDefault();
+			UE_LOG(LogInit, Log, TEXT("Core style set to default."));
 		}
 	}
 	else
 	{
 		FCoreStyle::ResetToDefault();
+		UE_LOG(LogInit, Log, TEXT("Core style set to default."));
 	}
 
 	{
@@ -1139,38 +1146,59 @@ int32 FEngineLoop::PreInit( const TCHAR* CmdLine )
 		// Initialize shader types before loading any shaders
 		InitializeShaderTypes();
 
+		UE_LOG(LogInit, Log, TEXT("Shader types initialized."));
+
 		// Load the global shaders.
 		if (GetGlobalShaderMap() == NULL && GIsRequestingExit)
 		{
 			// This means we can't continue without the global shader map.
+			UE_LOG(LogInit, Log, TEXT("Global shader map could not be verified."));
 			return 1;
 		}
 
+		UE_LOG(LogInit, Log, TEXT("Global shader map verified."));
+
 		// In order to be able to use short script package names get all script
 		// package names from ini files and register them with FPackageName system.
 		FPackageName::RegisterShortPackageNamesForUObjectModules();
 
+		UE_LOG(LogInit, Log, TEXT("Short package names for UObject modules registered."));
 
 		// Make sure all UObject classes are registered and default properties have been initialized
 		ProcessNewlyLoadedUObjects();
 
+		UE_LOG(LogInit, Log, TEXT("Newly loaded UObjects processed."));
+
 		// Default materials may have been loaded due to dependencies when loading
 		// classes and class default objects. If not, do so now.
 		UMaterialInterface::InitDefaultMaterials();
+
+		UE_LOG(LogInit, Log, TEXT("Default materials initialized."));
+
 		UMaterialInterface::AssertDefaultMaterialsExist();
+
+		UE_LOG(LogInit, Log, TEXT("Default materials exist."));
+
 		UMaterialInterface::AssertDefaultMaterialsPostLoaded();
+
+		UE_LOG(LogInit, Log, TEXT("Default materials are post loaded."));
 	}
 
 	// Tell the module manager is may now process newly-loaded UObjects when new C++ modules are loaded
 	FModuleManager::Get().StartProcessingNewlyLoadedObjects();
 
+	UE_LOG(LogInit, Log, TEXT("Newly loaded objects are now being processed."));
+
 	// load up the seek-free startup packages
 	if ( !FStartupPackages::LoadAll() )
 	{
 		// At least one startup package failed to load, return 1 to indicate an error
+		UE_LOG(LogInit, Log, TEXT("Failed to load all startup packages."));
 		return 1;
 	}
 
+	UE_LOG(LogInit, Log, TEXT("Startup packages have been loaded."));
+
 	// Setup GC optimizations
 	if (bIsSeekFreeDedicatedServer || bHasEditorToken)
 	{
@@ -1180,9 +1208,12 @@ int32 FEngineLoop::PreInit( const TCHAR* CmdLine )
 	if ( !LoadStartupCoreModules() )
 	{
 		// At least one startup module failed to load, return 1 to indicate an error
+		UE_LOG(LogInit, Log, TEXT("Startup core modules could not be loaded."));
 		return 1;
 	}
 
+	UE_LOG(LogInit, Log, TEXT("Startup core modules are loaded."));
+
 #if !UE_SERVER// && !UE_EDITOR
 	if (!IsRunningDedicatedServer() && !IsRunningCommandlet())
 	{
@@ -1199,10 +1230,15 @@ int32 FEngineLoop::PreInit( const TCHAR* CmdLine )
 	// Load up all modules that need to hook into the loading screen
 	if (!IProjectManager::Get().LoadModulesForProject(ELoadingPhase::PreLoadingScreen))
 	{
+		UE_LOG(LogInit, Log, TEXT("Failed to load modules for project in the PreLoadingScreen phase."));
 		return 1;
 	}
 
+	UE_LOG(LogInit, Log, TEXT("Loaded modules for project in the PreLoadingScreen phase."));
+
 	IPluginManager::Get().LoadModulesForEnabledPlugins(ELoadingPhase::PreLoadingScreen);
+
+	UE_LOG(LogInit, Log, TEXT("Loaded modules for plugins in the PreLoadingScreen phase."));
 	
 	
 #if !UE_SERVER
@@ -1220,25 +1256,40 @@ int32 FEngineLoop::PreInit( const TCHAR* CmdLine )
 		// it wouldn't have anything in it's .ini file)
 		GetMoviePlayer()->SetupLoadingScreenFromIni();
 
+		UE_LOG(LogInit, Log, TEXT("Loading screen setup from ini."));
+
 		GetMoviePlayer()->Initialize();
+
+		UE_LOG(LogInit, Log, TEXT("Movie player initialized."));
+
 		GetMoviePlayer()->PlayMovie();
+
+		UE_LOG(LogInit, Log, TEXT("Movie player playing."));
 	}
 #endif
 
+	UE_LOG(LogInit, Log, TEXT("Beginning Platform Post Init."));
+
 	// do any post appInit processing, before the render thread is started.
 	FPlatformMisc::PlatformPostInit();
 
+	UE_LOG(LogInit, Log, TEXT("Platform Post Init complete."));
+
 	if (GUseThreadedRendering)
 	{
 		StartRenderingThread();
+		UE_LOG(LogInit, Log, TEXT("Rendering thread started."));
 	}
 	
 	if ( !LoadStartupModules() )
 	{
 		// At least one startup module failed to load, return 1 to indicate an error
+		UE_LOG(LogInit, Log, TEXT("Failed to load startup modules."));
 		return 1;
 	}
 
+	UE_LOG(LogInit, Log, TEXT("Startup modules loaded."));
+
 	MarkObjectsToDisregardForGC(); 
 	GUObjectArray.CloseDisregardForGC();
 
