Platform; implement sensor events for Android. The implementation will use the accelerometer to derive tilt, rotation rate, gravity, and linear acceleration. If avialable gyroscope readings will get used for the rotation rate.

diff --git a/Engine/Source/Runtime/Core/Private/Android/AndroidInputInterface.cpp b/Engine/Source/Runtime/Core/Private/Android/AndroidInputInterface.cpp
--- a/Engine/Source/Runtime/Core/Private/Android/AndroidInputInterface.cpp
+++ b/Engine/Source/Runtime/Core/Private/Android/AndroidInputInterface.cpp
@@ -20,6 +20,9 @@ FDeferredAndroidMessage FAndroidInputInterface::DeferredMessages[MAX_DEFERRED_ME
 int32 FAndroidInputInterface::DeferredMessageQueueLastEntryIndex = 0;
 int32 FAndroidInputInterface::DeferredMessageQueueDroppedCount   = 0;
 
+TArray<FAndroidInputInterface::MotionData> FAndroidInputInterface::MotionDataStack
+	= TArray<FAndroidInputInterface::MotionData>();
+
 
 TSharedRef< FAndroidInputInterface > FAndroidInputInterface::Create(  const TSharedRef< FGenericApplicationMessageHandler >& InMessageHandler )
 {
@@ -165,6 +168,16 @@ void FAndroidInputInterface::SendControllerEvents()
 		OldControllerState = NewControllerState;
 	}
 
+	for (int i = 0; i < FAndroidInputInterface::MotionDataStack.Num(); ++i)
+	{
+		MotionData motion_data = FAndroidInputInterface::MotionDataStack[i];
+
+		MessageHandler->OnMotionDetected(
+			motion_data.Tilt, motion_data.RotationRate,
+			motion_data.Gravity, motion_data.Acceleration,
+			0);
+	}
+
 	for (int32 MessageIndex = 0; MessageIndex < FMath::Min(DeferredMessageQueueLastEntryIndex, MAX_DEFERRED_MESSAGE_QUEUE_SIZE); ++MessageIndex)
 	{
 		const FDeferredAndroidMessage& DeferredMessage = DeferredMessages[MessageIndex];
@@ -194,6 +207,8 @@ void FAndroidInputInterface::SendControllerEvents()
 	DeferredMessageQueueLastEntryIndex = 0;
 
 	FAndroidInputInterface::TouchInputStack.Empty(0);
+
+	FAndroidInputInterface::MotionDataStack.Empty();
 }
 
 void FAndroidInputInterface::QueueTouchInput(TArray<TouchInput> InTouchEvents)
@@ -273,3 +288,11 @@ void FAndroidInputInterface::DeferMessage(const FDeferredAndroidMessage& Deferre
 	}
 	DeferredMessages[Index] = DeferredMessage;
 }
+
+void FAndroidInputInterface::QueueMotionData(const FVector& Tilt, const FVector& RotationRate, const FVector& Gravity, const FVector& Acceleration)
+{
+	FScopeLock Lock(&TouchInputCriticalSection);
+
+	FAndroidInputInterface::MotionDataStack.Push(
+		MotionData { Tilt, RotationRate, Gravity, Acceleration });
+}
diff --git a/Engine/Source/Runtime/Core/Public/Android/AndroidInputInterface.h b/Engine/Source/Runtime/Core/Public/Android/AndroidInputInterface.h
--- a/Engine/Source/Runtime/Core/Public/Android/AndroidInputInterface.h
+++ b/Engine/Source/Runtime/Core/Public/Android/AndroidInputInterface.h
@@ -160,6 +160,8 @@ public:
 
 	static void DeferMessage(const FDeferredAndroidMessage& DeferredMessage);
 
+	static void QueueMotionData(const FVector& Tilt, const FVector& RotationRate, const FVector& Gravity, const FVector& Acceleration);
+
 private:
 
 	FAndroidInputInterface( const TSharedRef< FGenericApplicationMessageHandler >& InMessageHandler );
@@ -167,6 +169,14 @@ private:
 
 private:
 
+	struct MotionData
+	{
+		FVector Tilt;
+		FVector RotationRate;
+		FVector Gravity;
+		FVector Acceleration;
+	};
+
 	// protects the input stack
 	static FCriticalSection TouchInputCriticalSection;
 
@@ -184,5 +194,7 @@ private:
 	static int32 DeferredMessageQueueLastEntryIndex;
 	static int32 DeferredMessageQueueDroppedCount;
 
+	static TArray<MotionData> MotionDataStack;
+
 	TSharedRef< FGenericApplicationMessageHandler > MessageHandler;
 };
diff --git a/Engine/Source/Runtime/Launch/Private/Android/LaunchAndroid.cpp b/Engine/Source/Runtime/Launch/Private/Android/LaunchAndroid.cpp
--- a/Engine/Source/Runtime/Launch/Private/Android/LaunchAndroid.cpp
+++ b/Engine/Source/Runtime/Launch/Private/Android/LaunchAndroid.cpp
@@ -13,6 +13,8 @@
 #include <sys/resource.h>
 #include <dlfcn.h>
 #include "AndroidWindow.h"
+#include <android/sensor.h>
+#include "Core.h"
 
 // Function pointer for retrieving joystick events
 // Function has been part of the OS since Honeycomb, but only appeared in the
@@ -76,6 +78,9 @@ static void AndroidProcessEvents(struct android_app* state);
 //Event thread stuff
 static void* AndroidEventThreadWorker(void* param);
 
+// How often to process (read & dispatch) events, in seconds.
+static const float EventRefreshRate = 1.0f / 20.0f;
+
 //Android event callback functions
 static int32_t HandleInputCB(struct android_app* app, AInputEvent* event); //Touch and key input events
 static void OnAppCommandCB(struct android_app* app, int32_t cmd); //Lifetime events
@@ -83,6 +88,18 @@ static void OnAppCommandCB(struct android_app* app, int32_t cmd); //Lifetime eve
 bool GHasInterruptionRequest = false;
 bool GIsInterrupted = false;
 
+// Android sensor data management
+static ASensorManager * SensorManager = NULL;
+// Accelerometer (includes gravity), i.e. FMotionEvent::GetAcceleration.
+static const ASensor * SensorAccelerometer = NULL;
+// Gyroscope, i.e. FMotionEvent::GetRotationRate.
+static const ASensor * SensorGyroscope = NULL;
+static ASensorEventQueue * SensorQueue = NULL;
+// android.hardware.SensorManager.SENSOR_DELAY_GAME
+static const int32_t SensorDelayGame = 1;
+// Time decay sampling rate.
+static const float SampleDecayRate = 0.85f;
+
 void UpdateGameInterruptions()
 {
 	// Check for game suspension.
@@ -300,6 +317,25 @@ static void* AndroidEventThreadWorker( void* param )
 
 	FPlatformMisc::LowLevelOutputDebugString(L"Passed callback initialization");
 
+	// Acquire sensors
+	SensorManager = ASensorManager_getInstance();
+	if (NULL != SensorManager)
+	{
+		// Register for the various sensor events we want. Some
+		// may return NULL indicating that the sensor data is not
+		// available in the device. For those empty data will eventually
+		// get fed into the motion events.
+		SensorAccelerometer = ASensorManager_getDefaultSensor(
+			SensorManager, ASENSOR_TYPE_ACCELEROMETER);
+		SensorGyroscope = ASensorManager_getDefaultSensor(
+			SensorManager, ASENSOR_TYPE_GYROSCOPE);
+		// Create the queue for events to arrive.
+		SensorQueue = ASensorManager_createEventQueue(
+			SensorManager, state->looper, LOOPER_ID_USER, NULL, NULL);
+	}
+
+	FPlatformMisc::LowLevelOutputDebugString(L"Passed sensor initialization");
+
 	//continue to process events until the engine is shutting down
 	while (!GIsRequestingExit)
 	{
@@ -309,9 +345,8 @@ static void* AndroidEventThreadWorker( void* param )
 
 		AndroidProcessEvents(state);
 
-		float timeToSleep = 0.05f; //in seconds
-		sleep(timeToSleep);
-	}	
+		sleep(EventRefreshRate);
+	}
 
 	UE_LOG(LogAndroid, Log, TEXT("Exiting"));
 
@@ -326,14 +361,134 @@ static void AndroidProcessEvents(struct android_app* state)
 	int events;
 	struct android_poll_source* source;
 
+	// It's not possible to discern sequencing across sensors in
+	// Android. So we average out all the sensor events on one cycle
+	// and post a single motion sensor data point. We also need
+	// to synthesize additional information.
+	FVector current_accelerometer(0, 0, 0);
+	FVector current_gyroscope(0, 0, 0);
+	int32 current_accelerometer_sample_count = 0;
+	int32 current_gyroscope_sample_count = 0;
+	static FVector last_accelerometer(0, 0, 0);
+
 	while((ident = ALooper_pollAll(0, &fdesc, &events, (void**)&source)) >= 0)
 	{
 		// process this event
 		if (source)
 			source->process(state, source);
+
+		// process sensor events
+		if (ident == LOOPER_ID_USER)
+		{
+			if (NULL != SensorAccelerometer || NULL != SensorGyroscope)
+			{
+				ASensorEvent sensor_event;
+				while (ASensorEventQueue_getEvents(SensorQueue, &sensor_event, 1) > 0)
+				{
+					if (ASENSOR_TYPE_ACCELEROMETER == sensor_event.type)
+					{
+						current_accelerometer.X += sensor_event.acceleration.x;
+						current_accelerometer.Y += sensor_event.acceleration.y;
+						current_accelerometer.Z += sensor_event.acceleration.z;
+						current_accelerometer_sample_count += 1;
+					}
+					else if (ASENSOR_TYPE_GYROSCOPE == sensor_event.type)
+					{
+						current_gyroscope.X += sensor_event.vector.pitch;
+						current_gyroscope.Y += sensor_event.vector.azimuth;
+						current_gyroscope.Z += sensor_event.vector.roll;
+						current_gyroscope_sample_count += 1;
+					}
+				}
+			}
 		}
 	}
 
+	if (current_accelerometer_sample_count > 0)
+	{
+		// Do simple average of the samples we just got.
+		current_accelerometer /= float(current_accelerometer_sample_count);
+		last_accelerometer = current_accelerometer;
+	}
+	else
+	{
+		current_accelerometer = last_accelerometer;
+	}
+
+	if (current_gyroscope_sample_count > 0)
+	{
+		// Do simple average of the samples we just got.
+		current_gyroscope /= float(current_gyroscope_sample_count);
+	}
+
+	// If we have motion samples we generate the single event.
+	if (current_accelerometer_sample_count > 0 ||
+		current_gyroscope_sample_count > 0)
+	{
+		// The data we compose the motion event from.
+		FVector current_tilt(0, 0, 0);
+		FVector current_rotation_rate(0, 0, 0);
+		FVector current_gravity(0, 0, 0);
+		FVector current_acceleration(0, 0, 0);
+
+		// Buffered, historical, motion data.
+		static FVector last_tilt(0, 0, 0);
+		static FVector last_gravity(0, 0, 0);
+
+		// We use a low-pass filter to synthesize the gravity
+		// vector.
+		static bool first_acceleration_sample = true;
+		if (!first_acceleration_sample)
+		{
+			current_gravity
+				= last_gravity*SampleDecayRate
+				+ current_accelerometer*(1.0f - SampleDecayRate);
+		}
+		first_acceleration_sample = false;
+
+		// Calc the tilt from the accelerometer as it's not
+		// available directly.
+		FVector accelerometer_dir = -current_accelerometer.SafeNormal();
+		float current_pitch
+			= FMath::Atan2(accelerometer_dir.Y, accelerometer_dir.Z);
+		float current_roll
+			= -FMath::Atan2(accelerometer_dir.X, accelerometer_dir.Z);
+		current_tilt.X = current_pitch;
+		current_tilt.Y = 0;
+		current_tilt.Z = current_roll;
+
+		// And take out the gravity from the accel to get
+		// the linear acceleration.
+		current_acceleration = current_accelerometer - current_gravity;
+
+		if (current_gyroscope_sample_count > 0)
+		{
+			// The rotation rate is the what the gyroscope gives us.
+			current_rotation_rate = current_gyroscope;
+		}
+		else if (NULL == SensorGyroscope)
+		{
+			// If we don't have a gyroscope at all we need to calc a rotation
+			// rate from our calculated tilt and a delta.
+			current_rotation_rate = current_tilt - last_tilt;
+		}
+
+		// Finally record the motion event with all the data.
+		FAndroidInputInterface::QueueMotionData(current_tilt,
+			current_rotation_rate, current_gravity, current_acceleration);
+
+		// Update history values.
+		last_tilt = current_tilt;
+		last_gravity = current_gravity;
+
+		// UE_LOG(LogTemp, Log, TEXT("MOTION: tilt = %s, rotation-rate = %s, gravity = %s, acceleration = %s"),
+		//	*current_tilt.ToCompactString(),
+		//	*current_rotation_rate.ToCompactString(),
+		//	*current_gravity.ToCompactString(),
+		//	*current_acceleration.ToCompactString());
+	}
+}
+
 pthread_t G_AndroidEventThread;
 
 void android_main(struct android_app* state)
@@ -583,6 +738,13 @@ static void OnAppCommandCB(struct android_app* app, int32_t cmd)
 		// if the app lost focus, avoid unnecessary processing (like monitoring the accelerometer)
 		UE_LOG(LogAndroid, Log, TEXT("Case APP_CMD_LOST_FOCUS"));
 		FAppEventManager::GetInstance()->EnqueueAppEvent(APP_EVENT_STATE_WINDOW_LOST_FOCUS, NULL);
+		if (NULL != SensorQueue)
+		{
+			if (NULL != SensorAccelerometer)
+				ASensorEventQueue_disableSensor(SensorQueue, SensorAccelerometer);
+			if (NULL != SensorGyroscope)
+				ASensorEventQueue_disableSensor(SensorQueue, SensorGyroscope);
+		}
 		break;
 	case APP_CMD_GAINED_FOCUS:
 		/**
@@ -592,6 +754,19 @@ static void OnAppCommandCB(struct android_app* app, int32_t cmd)
 		// bring back a certain functionality, like monitoring the accelerometer
 		UE_LOG(LogAndroid, Log, TEXT("Case APP_CMD_GAINED_FOCUS"));
 		FAppEventManager::GetInstance()->EnqueueAppEvent(APP_EVENT_STATE_WINDOW_GAINED_FOCUS, NULL);
+		if (NULL != SensorQueue)
+		{
+			if (NULL != SensorAccelerometer)
+			{
+				ASensorEventQueue_enableSensor(SensorQueue, SensorAccelerometer);
+				ASensorEventQueue_setEventRate(SensorQueue, SensorAccelerometer, SensorDelayGame);
+			}
+			if (NULL != SensorGyroscope)
+			{
+				ASensorEventQueue_enableSensor(SensorQueue, SensorGyroscope);
+				ASensorEventQueue_setEventRate(SensorQueue, SensorGyroscope, SensorDelayGame);
+			}
+		}
 		break;
 	case APP_CMD_INPUT_CHANGED:
 		UE_LOG(LogAndroid, Log, TEXT("Case APP_CMD_INPUT_CHANGED"));
