Adding -testexit="Phrase" to aid automated testing in UAT

diff --git a/Engine/Source/Runtime/Launch/Private/LaunchEngineLoop.cpp b/Engine/Source/Runtime/Launch/Private/LaunchEngineLoop.cpp
--- a/Engine/Source/Runtime/Launch/Private/LaunchEngineLoop.cpp
+++ b/Engine/Source/Runtime/Launch/Private/LaunchEngineLoop.cpp
@@ -106,8 +106,47 @@ private:
 	bool bAllowLogVerbosity;
 };
 
+// Exits the game/editor if any of the specified phrases appears in the log output
+class FOutputDeviceTestExit : public FOutputDevice
+{
+	TArray<FString> ExitPhrases;
+public:
+	FOutputDeviceTestExit(const TArray<FString>& InExitPhrases)
+		: ExitPhrases(InExitPhrases)
+	{
+	}
+	virtual ~FOutputDeviceTestExit()
+	{
+	}
+	virtual void Serialize(const TCHAR* V, ELogVerbosity::Type Verbosity, const class FName& Category) OVERRIDE
+	{
+		if (!GIsRequestingExit)
+		{
+			for (auto& Phrase : ExitPhrases)
+			{
+				if (FCString::Stristr(V, *Phrase) && !FCString::Stristr(V, TEXT("-testexit=")))
+				{
+					if (GEngine != NULL)
+					{
+						if (GIsEditor)
+						{
+							GEngine->DeferredCommands.Add(TEXT("CLOSE_SLATE_MAINFRAME"));					
+						}
+						else
+						{
+							GEngine->Exec(NULL, TEXT("QUIT"));
+						}
+					}
+					break;
+				}
+			}
+		}
+	}
+};
+
 static TScopedPointer<FOutputDeviceConsole>	GScopedLogConsole;
 static TScopedPointer<FOutputDeviceStdOutput> GScopedStdOut;
+static TScopedPointer<FOutputDeviceTestExit> GScopedTestExit;
 
 /**
  * Initializes std out device and adds it to GLog
@@ -586,6 +625,22 @@ int32 FEngineLoop::PreInit( const TCHAR* CmdLine )
 		InitializeStdOutDevice();
 	}
 
+#if !UE_BUILD_SHIPPING
+	if (FPlatformProperties::SupportsQuit())
+	{
+		FString ExitPhrases;
+		if (FParse::Value(FCommandLine::Get(), TEXT("testexit="), ExitPhrases))
+		{
+			TArray<FString> ExitPhrasesList;
+			if (ExitPhrases.ParseIntoArray(&ExitPhrasesList, TEXT("+"), true) > 0)
+			{
+				GScopedTestExit = new FOutputDeviceTestExit(ExitPhrasesList);
+				GLog->AddOutputDevice(GScopedTestExit.GetOwnedPointer());
+			}
+		}
+	}
+#endif // !UE_BUILD_SHIPPING
+
 	// Set GGameName, based on the command line
 	if (LaunchSetGameName(CmdLine) == false)
 	{
