Fixing CIS

diff --git a/Engine/Source/Runtime/Launch/Private/LaunchEngineLoop.cpp b/Engine/Source/Runtime/Launch/Private/LaunchEngineLoop.cpp
--- a/Engine/Source/Runtime/Launch/Private/LaunchEngineLoop.cpp
+++ b/Engine/Source/Runtime/Launch/Private/LaunchEngineLoop.cpp
@@ -126,6 +126,7 @@ public:
 			{
 				if (FCString::Stristr(V, *Phrase) && !FCString::Stristr(V, TEXT("-testexit=")))
 				{
+#if WITH_ENGINE
 					if (GEngine != NULL)
 					{
 						if (GIsEditor)
@@ -137,6 +138,9 @@ public:
 							GEngine->Exec(NULL, TEXT("QUIT"));
 						}
 					}
+#else
+					FPlatformMisc::RequestExit(true);
+#endif
 					break;
 				}
 			}
