fix for not finding the right UE4Commandline.txt on Android
TTP334733
#ue4
#android

diff --git a/Engine/Source/Runtime/Launch/Private/Android/LaunchAndroid.cpp b/Engine/Source/Runtime/Launch/Private/Android/LaunchAndroid.cpp
--- a/Engine/Source/Runtime/Launch/Private/Android/LaunchAndroid.cpp
+++ b/Engine/Source/Runtime/Launch/Private/Android/LaunchAndroid.cpp
@@ -152,7 +152,7 @@ static void InitCommandLine()
 	FCommandLine::Set(TEXT(""));
 
 	// read in the command line text file from the sdcard if it exists
-	FString CommandLineFilePath = GFilePathBase + FString("/") + (GGameName[0] ? GGameName : TEXT("UE4Game")) + FString("/UE4CommandLine.txt");
+	FString CommandLineFilePath = GFilePathBase + FString("/") + (GGameName[0] ? GGameName : FPlatformProcess::ExecutableName()) + FString("/UE4CommandLine.txt");
 	FILE* CommandLineFile = fopen(TCHAR_TO_UTF8(*CommandLineFilePath), "r");
 	if(CommandLineFile == NULL)
 	{
