Offer to recompile missing binaries on startup if they are missing or out of date.

diff --git a/Engine/Source/Runtime/Launch/Launch.Build.cs b/Engine/Source/Runtime/Launch/Launch.Build.cs
--- a/Engine/Source/Runtime/Launch/Launch.Build.cs
+++ b/Engine/Source/Runtime/Launch/Launch.Build.cs
@@ -94,6 +94,7 @@ public class Launch : ModuleRules
 				new string[] {
 					"SourceControl",
 					"UnrealEd",
+					"DesktopPlatform"
 				}
 			);
 
diff --git a/Engine/Source/Runtime/Launch/Private/LaunchEngineLoop.cpp b/Engine/Source/Runtime/Launch/Private/LaunchEngineLoop.cpp
--- a/Engine/Source/Runtime/Launch/Private/LaunchEngineLoop.cpp
+++ b/Engine/Source/Runtime/Launch/Private/LaunchEngineLoop.cpp
@@ -66,6 +66,10 @@
 	static FFeedbackContextEditor UnrealEdWarn;
 #endif	// WITH_EDITOR
 
+#if UE_EDITOR
+	#include "DesktopPlatformModule.h"
+#endif
+
 #define LOCTEXT_NAMESPACE "LaunchEngineLoop"
 
 #if PLATFORM_WINDOWS
@@ -978,7 +982,11 @@ int32 FEngineLoop::PreInit( const TCHAR* CmdLine )
 
 	LoadPreInitModules();
 
-	AppInit();
+	// Start the application
+	if(!AppInit())
+	{
+		return 1;
+	}
 	
 #if WITH_ENGINE
 	// Initialize system settings before anyone tries to use it...
@@ -2222,7 +2230,7 @@ void FEngineLoop::OnSuspending(_In_ Platform::Object^ Sender, _In_ Windows::Appl
 /* FEngineLoop static interface
  *****************************************************************************/
 
-void FEngineLoop::AppInit( )
+bool FEngineLoop::AppInit( )
 {
 	// Output devices.
 	GError = FPlatformOutputDevices::GetError();
@@ -2299,6 +2307,39 @@ void FEngineLoop::AppInit( )
 	// init config system
 	FConfigCacheIni::InitializeConfigSystem();
 
+	// Check whether the project or any of its plugins are missing or are out of date
+#if UE_EDITOR
+	if(!GIsBuildMachine && FPaths::IsProjectFilePathSet())
+	{
+		if(!IProjectManager::Get().AreProjectModulesUpToDate() || !IPluginManager::Get().AreEnabledPluginModulesUpToDate())
+		{
+			int32 Result = FPlatformMisc::MessageBoxExt(EAppMsgType::YesNoCancel, TEXT("Project modules are missing or out of date. Would you like to recompile them?"), TEXT("Question"));
+			if(Result == EAppReturnType::Yes)
+			{
+				FFeedbackContext *Context = FDesktopPlatformModule::Get()->GetNativeFeedbackContext();
+
+				// Try to compile it
+				Context->BeginSlowTask(FText::FromString(TEXT("Project is out of date, recompiling...")), true, true);
+				bool bCompileResult = FDesktopPlatformModule::Get()->CompileGameProject(FPaths::RootDir(), FPaths::GetProjectFilePath(), Context);
+				Context->EndSlowTask();
+
+				// Check it succeeded
+				if(!bCompileResult || !IProjectManager::Get().AreProjectModulesUpToDate() || !IPluginManager::Get().AreEnabledPluginModulesUpToDate())
+				{
+					if(FPlatformMisc::MessageBoxExt(EAppMsgType::YesNo, TEXT("Game code couldn't be compiled. Continue trying to start anyway?"), TEXT("Error")) == EAppReturnType::No)
+					{
+						return false;
+					}
+				}
+			}
+			else if(Result == EAppReturnType::Cancel)
+			{
+				return false;
+			}
+		}
+	}
+#endif
+
 	// Load "pre-init" plugin modules
 	if (IProjectManager::Get().LoadModulesForProject(ELoadingPhase::PostConfigInit))
 	{
@@ -2456,6 +2497,7 @@ void FEngineLoop::AppInit( )
 
 	// Init other systems.
 	FCoreDelegates::OnInit.Broadcast();
+	return true;
 }
 
 
diff --git a/Engine/Source/Runtime/Launch/Public/LaunchEngineLoop.h b/Engine/Source/Runtime/Launch/Public/LaunchEngineLoop.h
--- a/Engine/Source/Runtime/Launch/Public/LaunchEngineLoop.h
+++ b/Engine/Source/Runtime/Launch/Public/LaunchEngineLoop.h
@@ -102,7 +102,7 @@ public:
 	/**
 	 * Initializes the application.
 	 */
-	static void AppInit( );
+	static bool AppInit( );
 
 	/**
 	 * Prepares the application for shutdown.
