Made cook on the fly server on Mac process Cocoa events so that it can quit properly

diff --git a/Engine/Source/Editor/UnrealEd/Private/Commandlets/CookCommandlet.cpp b/Engine/Source/Editor/UnrealEd/Private/Commandlets/CookCommandlet.cpp
--- a/Engine/Source/Editor/UnrealEd/Private/Commandlets/CookCommandlet.cpp
+++ b/Engine/Source/Editor/UnrealEd/Private/Commandlets/CookCommandlet.cpp
@@ -165,6 +165,11 @@ bool UCookCommandlet::CookOnTheFly( FGuid InstanceId, int32 Timeout, bool bForce
 				}
 			}
 
+#if PLATFORM_MAC
+			// On Mac we need to process Cocoa events so that the console window for CookOnTheFlyServer is interactive
+			FPlatformMisc::PumpMessages(true);
+#endif
+
 			// update task graph
 			FTaskGraphInterface::Get().ProcessThreadUntilIdle(ENamedThreads::GameThread);
 
diff --git a/Engine/Source/Runtime/Launch/Private/Mac/LaunchMac.cpp b/Engine/Source/Runtime/Launch/Private/Mac/LaunchMac.cpp
--- a/Engine/Source/Runtime/Launch/Private/Mac/LaunchMac.cpp
+++ b/Engine/Source/Runtime/Launch/Private/Mac/LaunchMac.cpp
@@ -95,11 +95,18 @@ void EngineCrashHandler(const FGenericCrashContext & GenericContext)
 	{
 		if (GIsEditor)
 		{
-			GEngine->DeferredCommands.Add( TEXT("CLOSE_SLATE_MAINFRAME") );
+			if (IsRunningCommandlet())
+			{
+				GIsRequestingExit = true;
+			}
+			else
+			{
+				GEngine->DeferredCommands.Add(TEXT("CLOSE_SLATE_MAINFRAME"));
+			}
 		}
 		else
 		{
-			GEngine->DeferredCommands.Add( TEXT("EXIT") );
+			GEngine->DeferredCommands.Add(TEXT("EXIT"));
 		}
 	}
 }
