Linux: Make abnormal exit more obvious to user.

Alleviates TTP #335935 where passing an incorrect path project confused the user.

diff --git a/Engine/Source/Runtime/Launch/Private/LaunchEngineLoop.cpp b/Engine/Source/Runtime/Launch/Private/LaunchEngineLoop.cpp
--- a/Engine/Source/Runtime/Launch/Private/LaunchEngineLoop.cpp
+++ b/Engine/Source/Runtime/Launch/Private/LaunchEngineLoop.cpp
@@ -912,6 +912,7 @@ int32 FEngineLoop::PreInit( const TCHAR* CmdLine )
 		if ( !IProjectManager::Get().LoadProjectFile(FPaths::GetProjectFilePath()) )
 		{
 			// The project file was invalid or saved with a newer version of the engine. Exit.
+			UE_LOG(LogInit, Warning, TEXT("Could not find a valid project file, the engine will exit now."));
 			return 1;
 		}
 	}
diff --git a/Engine/Source/Runtime/Launch/Private/Linux/LaunchLinux.cpp b/Engine/Source/Runtime/Launch/Private/Linux/LaunchLinux.cpp
--- a/Engine/Source/Runtime/Launch/Private/Linux/LaunchLinux.cpp
+++ b/Engine/Source/Runtime/Launch/Private/Linux/LaunchLinux.cpp
@@ -180,5 +180,9 @@ int main(int argc, char *argv[])
 		}
 	}
 
+	if (ErrorLevel)
+	{
+		printf("Exiting abnormally (error code: %d)\n", ErrorLevel);
+	}
 	return ErrorLevel;
 }
