fix for black screen on Razr
#ue4
#android

diff --git a/Engine/Source/Runtime/Launch/Private/Android/AndroidEventManager.cpp b/Engine/Source/Runtime/Launch/Private/Android/AndroidEventManager.cpp
--- a/Engine/Source/Runtime/Launch/Private/Android/AndroidEventManager.cpp
+++ b/Engine/Source/Runtime/Launch/Private/Android/AndroidEventManager.cpp
@@ -163,6 +163,8 @@ void FAppEventManager::HandleWindowCreated(void* InWindow)
 
 		rc = pthread_mutex_unlock(&MainMutex);
 		check(rc == 0);
+
+		EnqueueAppEvent(APP_EVENT_STATE_WINDOW_CREATED, InWindow );
 	}
 }
 
