Check to load SkillSystemEditor as early as possible to avoid issues with statically loaded resources neeeding it

diff --git a/Engine/Source/Runtime/Launch/Private/LaunchEngineLoop.cpp b/Engine/Source/Runtime/Launch/Private/LaunchEngineLoop.cpp
--- a/Engine/Source/Runtime/Launch/Private/LaunchEngineLoop.cpp
+++ b/Engine/Source/Runtime/Launch/Private/LaunchEngineLoop.cpp
@@ -1696,6 +1696,19 @@ bool FEngineLoop::LoadStartupCoreModules()
 		FModuleManager::Get().LoadModule(TEXT("BehaviorTreeEditor"));
 	}
 
+	// -----------------------------------------------------
+
+	// HACK: load SkillSystem editor as early as possible for statically initialized assets (non cooked BT assets needs it)
+	// cooking needs this module too
+	bool bSkillSystemEditorEnabled = false;
+	GConfig->GetBool(TEXT("SkillSystemEd"), TEXT("SKillSystemEditorEnabled"), bSkillSystemEditorEnabled, GEngineIni);
+	if (bSkillSystemEditorEnabled)
+	{
+		FModuleManager::Get().LoadModule(TEXT("SkillSystemEditor"));
+	}
+
+	// -----------------------------------------------------
+
 	// HACK: load EQS editor as early as possible for statically initialized assets (non cooked EQS assets needs it)
 	// cooking needs this module too
 	bool bEnvironmentQueryEditor = false;
