Added code to ensure OnlineBlueprintSupport is loaded during the LoadStartupCoreModules to ensure BP only projects that have online functionality work correctly.

TTP 333996 - CRITICAL: Regression: TAPPYCHICKEN: Map does not render in game

diff --git a/Engine/Source/Runtime/Launch/Private/LaunchEngineLoop.cpp b/Engine/Source/Runtime/Launch/Private/LaunchEngineLoop.cpp
--- a/Engine/Source/Runtime/Launch/Private/LaunchEngineLoop.cpp
+++ b/Engine/Source/Runtime/Launch/Private/LaunchEngineLoop.cpp
@@ -1718,6 +1718,9 @@ bool FEngineLoop::LoadStartupCoreModules()
 		FModuleManager::Get().LoadModule(TEXT("EnvironmentQueryEditor"));
 	}
 
+	// We need this for blueprint projects that have online functionality.
+	FModuleManager::Get().LoadModule(TEXT("OnlineBlueprintSupport"));
+
 #endif //(WITH_EDITOR && !(UE_BUILD_SHIPPING || UE_BUILD_TEST))
 
 	return bSuccess;
