Include Commandlet.h in LaunchEngineLoop.cpp

diff --git a/Engine/Source/Runtime/Launch/Private/LaunchEngineLoop.cpp b/Engine/Source/Runtime/Launch/Private/LaunchEngineLoop.cpp
--- a/Engine/Source/Runtime/Launch/Private/LaunchEngineLoop.cpp
+++ b/Engine/Source/Runtime/Launch/Private/LaunchEngineLoop.cpp
@@ -40,6 +40,8 @@
 	#include "Online.h"
 	#include "PlatformFeatures.h"
 	#include "DeviceProfiles/DeviceProfileManager.h"
+	#include "Commandlets/Commandlet.h"
+
 #if !UE_SERVER
 	#include "SlateRHIRendererModule.h"
 #endif
