Added FPlatformMisc::PlatformPostInit(); call to LaunchEngineLoop for programs that are not compiled with WITH_ENGINE

diff --git a/Engine/Source/Runtime/Launch/Private/LaunchEngineLoop.cpp b/Engine/Source/Runtime/Launch/Private/LaunchEngineLoop.cpp
--- a/Engine/Source/Runtime/Launch/Private/LaunchEngineLoop.cpp
+++ b/Engine/Source/Runtime/Launch/Private/LaunchEngineLoop.cpp
@@ -1586,9 +1586,10 @@ int32 FEngineLoop::PreInit( const TCHAR* CmdLine )
 #endif
 
 
-#else
+#else // WITH_ENGINE
 	EndInitTextLocalization();
-#endif
+	FPlatformMisc::PlatformPostInit();
+#endif // WITH_ENGINE
 
 	//run automation smoke tests now that everything is setup to run
 	FAutomationTestFramework::GetInstance().RunSmokeTests();
