Implement invites for Xbox Live.
The Live OSS will now trigger the OnSessionInviteAccepted delegate when a user accepts a game invite.
Known issue: it will always pass in a LocalUserNum of 0 because we don't get specific information about which user accepted the invite from the OS.

diff --git a/Engine/Source/Runtime/Launch/Launch.Build.cs b/Engine/Source/Runtime/Launch/Launch.Build.cs
--- a/Engine/Source/Runtime/Launch/Launch.Build.cs
+++ b/Engine/Source/Runtime/Launch/Launch.Build.cs
@@ -78,6 +78,12 @@ public class Launch : ModuleRules
 			// LaunchEngineLoop.cpp does a LoadModule() on OnlineSubsystem and OnlineSubsystemUtils when compiled WITH_ENGINE, so they must be marked as dependencies so that they get compiled and cleaned
 			DynamicallyLoadedModuleNames.Add("OnlineSubsystem");
 			DynamicallyLoadedModuleNames.Add("OnlineSubsystemUtils");
+
+			// XboxOneLauch calls into OnlineSubsystemLive to handle invites
+			if (Target.Platform == UnrealTargetPlatform.XboxOne)
+			{
+				PrivateDependencyModuleNames.Add("OnlineSubsystemLive");
+			}
 		}
 
 		if (Target.Configuration != UnrealTargetConfiguration.Shipping)
diff --git a/Engine/Source/Runtime/Online/OnlineSubsystem/Public/OnlineSessionSettings.h b/Engine/Source/Runtime/Online/OnlineSubsystem/Public/OnlineSessionSettings.h
--- a/Engine/Source/Runtime/Online/OnlineSubsystem/Public/OnlineSessionSettings.h
+++ b/Engine/Source/Runtime/Online/OnlineSubsystem/Public/OnlineSessionSettings.h
@@ -332,17 +332,6 @@ public:
 	{
 	}
 
-	/** Copy Constructor */
-	FOnlineSession(const FOnlineSession& Src) :
-		OwningUserId(Src.OwningUserId),
-		OwningUserName(Src.OwningUserName),
-		SessionSettings(Src.SessionSettings),
-		NumOpenPrivateConnections(Src.NumOpenPrivateConnections),
-		NumOpenPublicConnections(Src.NumOpenPublicConnections)
-	{
-		// NOTE: SessionInfo is copied manually per platform
-	}
-
 	virtual ~FOnlineSession() {}
 };
 
