Rework Xbox One invites to be more robust.
Now uses party events instead of protocol activation to handle invites.
Now reports the correct LocalUserNum to the OnSessionInviteAccepted delegate.
The Launch module no longer depends on the OnlineSubsystemLive module on Xbox One.

diff --git a/Engine/Source/Runtime/Launch/Launch.Build.cs b/Engine/Source/Runtime/Launch/Launch.Build.cs
--- a/Engine/Source/Runtime/Launch/Launch.Build.cs
+++ b/Engine/Source/Runtime/Launch/Launch.Build.cs
@@ -78,12 +78,6 @@ public class Launch : ModuleRules
 			// LaunchEngineLoop.cpp does a LoadModule() on OnlineSubsystem and OnlineSubsystemUtils when compiled WITH_ENGINE, so they must be marked as dependencies so that they get compiled and cleaned
 			DynamicallyLoadedModuleNames.Add("OnlineSubsystem");
 			DynamicallyLoadedModuleNames.Add("OnlineSubsystemUtils");
-
-			// XboxOneLauch calls into OnlineSubsystemLive to handle invites
-			if (Target.Platform == UnrealTargetPlatform.XboxOne)
-			{
-				PrivateDependencyModuleNames.Add("OnlineSubsystemLive");
-			}
 		}
 
 		if (Target.Configuration != UnrealTargetConfiguration.Shipping)
