Compile error fix

diff --git a/Engine/Source/Runtime/Launch/Private/Android/AndroidJNI.cpp b/Engine/Source/Runtime/Launch/Private/Android/AndroidJNI.cpp
--- a/Engine/Source/Runtime/Launch/Private/Android/AndroidJNI.cpp
+++ b/Engine/Source/Runtime/Launch/Private/Android/AndroidJNI.cpp
@@ -61,7 +61,7 @@ private:
 	}
 };
 
-template<> uint32 TThreadSingleton<FJNIHelper>::TlsSlot = 0;
+DECLARE_THREAD_SINGLETON( FJNIHelper );
 
 JNIEnv* GetJavaEnv(bool bRequireGlobalThis)
 {
