Avoid init errors for not-yet-used FJNIHelper class.

diff --git a/Engine/Source/Runtime/Launch/Private/Android/AndroidJNI.cpp b/Engine/Source/Runtime/Launch/Private/Android/AndroidJNI.cpp
--- a/Engine/Source/Runtime/Launch/Private/Android/AndroidJNI.cpp
+++ b/Engine/Source/Runtime/Launch/Private/Android/AndroidJNI.cpp
@@ -9,12 +9,16 @@
 
 #define JNI_CURRENT_VERSION JNI_VERSION_1_6
 
+#define USE_JNI_HELPER 0
+
 JavaVM* GJavaVM;
 jobject GJavaGlobalThis = NULL;
 
 extern FString GFilePathBase;
 extern bool GOBBinAPK;
 
+#if USE_JNI_HELPER
+
 //////////////////////////////////////////////////////////////////////////
 // FJNIHelper
 
@@ -63,10 +67,12 @@ private:
 
 DECLARE_THREAD_SINGLETON( FJNIHelper );
 
+#endif // USE_JNI_HELPER
+
 JNIEnv* GetJavaEnv(bool bRequireGlobalThis)
 {
 	//@TODO: ANDROID: Remove the other version if the helper works well
-#if 0
+#if USE_JNI_HELPER
 	if (!bRequireGlobalThis || (GJavaGlobalThis != nullptr))
 	{
 		return FJNIHelper::GetEnvironment();
