Revised handling of incorrect commandlet. Now shows console and exits if commandlet is invalid.

ttp 337112 - CRITICAL: EDITOR: CRASH: Running invalid commandlet

diff --git a/Engine/Source/Runtime/Launch/Private/LaunchEngineLoop.cpp b/Engine/Source/Runtime/Launch/Private/LaunchEngineLoop.cpp
--- a/Engine/Source/Runtime/Launch/Private/LaunchEngineLoop.cpp
+++ b/Engine/Source/Runtime/Launch/Private/LaunchEngineLoop.cpp
@@ -1332,7 +1332,13 @@ int32 FEngineLoop::PreInit( const TCHAR* CmdLine )
 			CommandletClass = FindObject<UClass>(ANY_PACKAGE,*Token,false);
 			if (!CommandletClass)
 			{
-				UE_LOG(LogInit, Fatal,TEXT("%s looked like a commandlet, but we could not find the class."), *Token);
+				if (GLogConsole && !GIsSilent)
+				{
+					GLogConsole->Show(true);
+				}
+				UE_LOG(LogInit, Error, TEXT("%s looked like a commandlet, but we could not find the class."), *Token);
+				GIsRequestingExit = true;
+				return 1;
 			}
 
 			extern bool GIsConsoleExecutable;
