Ensure that ParseCommandline is called in the case of GIsEditor being true
(Fixes, among other things, sound not being initialized for commandlets)

diff --git a/Engine/Source/Runtime/Launch/Private/LaunchEngineLoop.cpp b/Engine/Source/Runtime/Launch/Private/LaunchEngineLoop.cpp
--- a/Engine/Source/Runtime/Launch/Private/LaunchEngineLoop.cpp
+++ b/Engine/Source/Runtime/Launch/Private/LaunchEngineLoop.cpp
@@ -1396,6 +1396,8 @@ int32 FEngineLoop::PreInit( const TCHAR* CmdLine )
 
 					GEngine = GEditor = ConstructObject<UEditorEngine>( EditorEngineClass );
 
+					GEngine->ParseCommandline();
+
 					UE_LOG(LogInit, Log, TEXT("Initializing Editor Engine..."));
 					GEditor->InitEditor(this);
 					UE_LOG(LogInit, Log, TEXT("Initializing Editor Engine Completed"));
