Added FMacPlatformMisc::NormalizePath() and fixed a related bug in LaunchEngineLoop, so that it's aware that the path returned by FPaths::GetProjectFilePath() is normalized.

diff --git a/Engine/Source/Runtime/Core/Private/Mac/MacPlatformMisc.cpp b/Engine/Source/Runtime/Core/Private/Mac/MacPlatformMisc.cpp
--- a/Engine/Source/Runtime/Core/Private/Mac/MacPlatformMisc.cpp
+++ b/Engine/Source/Runtime/Core/Private/Mac/MacPlatformMisc.cpp
@@ -995,6 +995,16 @@ int32 FMacPlatformMisc::NumberOfCoresIncludingHyperthreads()
 	return NumberOfCores;
 }
 
+void FMacPlatformMisc::NormalizePath(FString& InPath)
+{
+	const bool bAppendSlash = InPath[InPath.Len() - 1] == '/'; // NSString will remove the trailing slash, if present, so we need to restore it after conversion
+	InPath = [[InPath.GetNSString() stringByStandardizingPath] stringByResolvingSymlinksInPath];
+	if (bAppendSlash)
+	{
+		InPath += TEXT("/");
+	}
+}
+
 #include "ModuleManager.h"
 
 void FMacPlatformMisc::LoadPreInitModules()
diff --git a/Engine/Source/Runtime/Core/Public/Mac/MacPlatformMisc.h b/Engine/Source/Runtime/Core/Public/Mac/MacPlatformMisc.h
--- a/Engine/Source/Runtime/Core/Public/Mac/MacPlatformMisc.h
+++ b/Engine/Source/Runtime/Core/Public/Mac/MacPlatformMisc.h
@@ -58,6 +58,7 @@ struct CORE_API FMacPlatformMisc : public FGenericPlatformMisc
 	static int32 NumberOfCores();
 	static int32 NumberOfCoresIncludingHyperthreads();
 	static void LoadPreInitModules();
+	static void NormalizePath(FString& InPath);
 
 
 	/** 
diff --git a/Engine/Source/Runtime/Launch/Private/LaunchEngineLoop.cpp b/Engine/Source/Runtime/Launch/Private/LaunchEngineLoop.cpp
--- a/Engine/Source/Runtime/Launch/Private/LaunchEngineLoop.cpp
+++ b/Engine/Source/Runtime/Launch/Private/LaunchEngineLoop.cpp
@@ -767,8 +767,12 @@ int32 FEngineLoop::PreInit( const TCHAR* CmdLine )
 	// VC++ tends to do this in its "external tools" config
 	Token = Token.Trim();
 
+	// Path returned by FPaths::GetProjectFilePath() is normalized, so may have symlinks and ~ resolved and may differ from the original path to .uproject passed in the command line
+	FString NormalizedToken = Token;
+	FPaths::NormalizeFilename(NormalizedToken);
+
 	const bool bFirstTokenIsGameName = (FApp::HasGameName() && Token == GGameName);
-	const bool bFirstTokenIsGameProjectFilePath = (FPaths::IsProjectFilePathSet() && Token.Replace(TEXT("\\"), TEXT("/")) == FPaths::GetProjectFilePath());
+	const bool bFirstTokenIsGameProjectFilePath = (FPaths::IsProjectFilePathSet() && NormalizedToken == FPaths::GetProjectFilePath());
 	const bool bFirstTokenIsGameProjectFileShortName = (FPaths::IsProjectFilePathSet() && Token == FPaths::GetCleanFilename(FPaths::GetProjectFilePath()));
 
 	if (bFirstTokenIsGameName || bFirstTokenIsGameProjectFilePath || bFirstTokenIsGameProjectFileShortName)
