A Shipping build with logging enabled will now print messages to the Visual Studio output window

https://udn.unrealengine.com/questions/208357/use-logging-in-shipping1-not-enabling-logging-in-s.html

This was failing the IsDebuggerPresent() check used when adding FOutputDeviceDebug, because IsDebuggerPresent() is hard-coded to return 0 when UE_BUILD_SHIPPING is set.

This change always adds FOutputDeviceDebug for a Shipping build with logging explicitly enabled.

Also fixed a minor whitespace issue in Build.h.

diff --git a/Engine/Source/Runtime/Launch/Private/LaunchEngineLoop.cpp b/Engine/Source/Runtime/Launch/Private/LaunchEngineLoop.cpp
--- a/Engine/Source/Runtime/Launch/Private/LaunchEngineLoop.cpp
+++ b/Engine/Source/Runtime/Launch/Private/LaunchEngineLoop.cpp
@@ -2288,7 +2288,8 @@ bool FEngineLoop::AppInit( )
 	}
 
 	// Only create debug output device if a debugger is attached or we're running on a console or build machine
-	if (!FPlatformProperties::SupportsWindowedMode() || FPlatformMisc::IsDebuggerPresent() || GIsBuildMachine)
+	// A shipping build with logging explicitly enabled will fail the IsDebuggerPresent() check, but we still need to add the debug output device for logging purposes
+	if (!FPlatformProperties::SupportsWindowedMode() || FPlatformMisc::IsDebuggerPresent() || (UE_BUILD_SHIPPING && !NO_LOGGING) || GIsBuildMachine)
 	{
 		GLog->AddOutputDevice( new FOutputDeviceDebug() );
 	}
