- IOS console history fixes (loading from user settings, and reordering the array so recent items are always at the end)

diff --git a/Engine/Source/Runtime/Core/Private/IOS/IOSAppDelegate.cpp b/Engine/Source/Runtime/Core/Private/IOS/IOSAppDelegate.cpp
--- a/Engine/Source/Runtime/Core/Private/IOS/IOSAppDelegate.cpp
+++ b/Engine/Source/Runtime/Core/Private/IOS/IOSAppDelegate.cpp
@@ -435,10 +435,14 @@ void InstallSignalHandlers()
 
 	self.CommandLineParseTimer = [NSTimer scheduledTimerWithTimeInterval:0.01f target:self selector:@selector(NoUrlCommandLine) userInfo:nil repeats:NO];
 #if !UE_BUILD_SHIPPING
-	self.ConsoleHistoryValues = [[NSUserDefaults standardUserDefaults] objectForKey:@"ConsoleHistory"];
-	if (self.ConsoleHistoryValues == nil)
+	// make a history buffer
+	self.ConsoleHistoryValues = [[NSMutableArray alloc] init];
+
+	// load saved history from disk
+	NSArray* SavedHistory = [[NSUserDefaults standardUserDefaults] objectForKey:@"ConsoleHistory"];
+	if (SavedHistory != nil)
 	{
-		self.ConsoleHistoryValues = [[NSMutableArray alloc] init];
+		[self.ConsoleHistoryValues addObjectsFromArray:SavedHistory];
 	}
 	self.ConsoleHistoryValuesIndex = -1;
 #endif
diff --git a/Engine/Source/Runtime/Launch/Private/IOS/IOSAppDelegateConsoleHandling.cpp b/Engine/Source/Runtime/Launch/Private/IOS/IOSAppDelegateConsoleHandling.cpp
--- a/Engine/Source/Runtime/Launch/Private/IOS/IOSAppDelegateConsoleHandling.cpp
+++ b/Engine/Source/Runtime/Launch/Private/IOS/IOSAppDelegateConsoleHandling.cpp
@@ -70,14 +70,17 @@
 			}
 		];
 
-		// add the command to the command history if it's unique
-		if (ExistingCommand == NSNotFound)
+		// remove the existing one, so we can move it to the end
+		if (ExistingCommand != NSNotFound)
 		{
-			[self.ConsoleHistoryValues addObject:ConsoleCommand];
-			// save to local storage
-			[[NSUserDefaults standardUserDefaults] setObject:self.ConsoleHistoryValues forKey:@"ConsoleHistory"];
-			[[NSUserDefaults standardUserDefaults] synchronize];
+			[self.ConsoleHistoryValues removeObjectAtIndex:ExistingCommand];
 		}
+
+		// add the command to the end
+		[self.ConsoleHistoryValues addObject:ConsoleCommand];
+		// save to local storage
+		[[NSUserDefaults standardUserDefaults] setObject:self.ConsoleHistoryValues forKey:@"ConsoleHistory"];
+		[[NSUserDefaults standardUserDefaults] synchronize];
 	}
 }
 
