Fix for android using power in the background.

diff --git a/Engine/Source/Runtime/Launch/Private/Android/AndroidEventManager.cpp b/Engine/Source/Runtime/Launch/Private/Android/AndroidEventManager.cpp
--- a/Engine/Source/Runtime/Launch/Private/Android/AndroidEventManager.cpp
+++ b/Engine/Source/Runtime/Launch/Private/Android/AndroidEventManager.cpp
@@ -120,6 +120,11 @@ void FAppEventManager::Tick()
 			bRunning = false;
 		}
 	}
+
+	if (!bRunning && FirstInitialized)
+	{
+		EventHandlerEvent->Wait();
+	}
 }
 
 FAppEventManager::FAppEventManager():
@@ -168,6 +173,11 @@ void FAppEventManager::HandleWindowCreated(void* InWindow)
 	}
 }
 
+void FAppEventManager::SetEventHandlerEvent(FEvent* InEventHandlerEvent)
+{
+	EventHandlerEvent = InEventHandlerEvent;
+}
+
 void FAppEventManager::PauseRendering()
 {
 	if(GUseThreadedRendering )
diff --git a/Engine/Source/Runtime/Launch/Private/Android/AndroidEventManager.h b/Engine/Source/Runtime/Launch/Private/Android/AndroidEventManager.h
--- a/Engine/Source/Runtime/Launch/Private/Android/AndroidEventManager.h
+++ b/Engine/Source/Runtime/Launch/Private/Android/AndroidEventManager.h
@@ -43,6 +43,7 @@ public:
 	
 	void Tick();
 	void EnqueueAppEvent(EAppEventState InState, void* InData = NULL);
+	void SetEventHandlerEvent(FEvent* InEventHandlerEvent);
 	void HandleWindowCreated(void* InWindow);
 	bool IsGamePaused();
 
@@ -66,6 +67,8 @@ private:
 	
 	pthread_mutex_t MainMutex;
 
+	FEvent* EventHandlerEvent;			// triggered every time the event handler thread fires off an event
+
 	//states
 	bool FirstInitialized;
 	bool bCreateWindow;
diff --git a/Engine/Source/Runtime/Launch/Private/Android/LaunchAndroid.cpp b/Engine/Source/Runtime/Launch/Private/Android/LaunchAndroid.cpp
--- a/Engine/Source/Runtime/Launch/Private/Android/LaunchAndroid.cpp
+++ b/Engine/Source/Runtime/Launch/Private/Android/LaunchAndroid.cpp
@@ -99,6 +99,10 @@ static ASensorEventQueue * SensorQueue = NULL;
 static const int32_t SensorDelayGame = 1;
 // Time decay sampling rate.
 static const float SampleDecayRate = 0.85f;
+// Event for coordinating pausing of the main and event handling threads to prevent background spinning
+static FEvent* EventHandlerEvent = NULL;
+
+
 
 void UpdateGameInterruptions()
 {
@@ -249,6 +253,10 @@ int32 AndroidMain(struct android_app* state)
 	InitCommandLine();
 	FPlatformMisc::LowLevelOutputDebugStringf(TEXT("Final commandline: %s\n"), FCommandLine::Get());
 
+	EventHandlerEvent = FPlatformProcess::CreateSynchEvent();
+	FPlatformMisc::LowLevelOutputDebugString(L"Created sync event");
+	FAppEventManager::GetInstance()->SetEventHandlerEvent(EventHandlerEvent);
+
 	// initialize the engine
 	GEngineLoop.PreInit(0, NULL, FCommandLine::Get());
 
@@ -342,13 +350,11 @@ static void* AndroidEventThreadWorker( void* param )
 	//continue to process events until the engine is shutting down
 	while (!GIsRequestingExit)
 	{
-		static int count = 1;
-		//LOGD("Event process loop #%d", count);
-		count++;
+		FPlatformMisc::LowLevelOutputDebugString(L"AndroidEventThreadWorker");
 
 		AndroidProcessEvents(state);
 
-		sleep(EventRefreshRate);
+		sleep(EventRefreshRate);		// this is really 0 since it takes int seconds.
 	}
 
 	UE_LOG(LogAndroid, Log, TEXT("Exiting"));
@@ -374,11 +380,16 @@ static void AndroidProcessEvents(struct android_app* state)
 	int32 current_gyroscope_sample_count = 0;
 	static FVector last_accelerometer(0, 0, 0);
 
-	while((ident = ALooper_pollAll(0, &fdesc, &events, (void**)&source)) >= 0)
+	while((ident = ALooper_pollAll(-1, &fdesc, &events, (void**)&source)) >= 0)
+	// while((ident = ALooper_pollAll(0, &fdesc, &events, (void**)&source)) >= 0)
 	{
+		FPlatformMisc::LowLevelOutputDebugString(L"AndroidProcessEvents");
+
 		// process this event
 		if (source)
+		{
 			source->process(state, source);
+		}
 
 		// process sensor events
 		if (ident == LOOPER_ID_USER)
@@ -860,6 +871,9 @@ static void OnAppCommandCB(struct android_app* app, int32_t cmd)
 		FAppEventManager::GetInstance()->EnqueueAppEvent(APP_EVENT_STATE_ON_PAUSE);
 		break;
 	}
+
+	if ( EventHandlerEvent )
+		EventHandlerEvent->Trigger();
 }
 
 //Native-defined functions
