- removed massive AndroidProcessingEvents spam

diff --git a/Engine/Source/Runtime/Launch/Private/Android/LaunchAndroid.cpp b/Engine/Source/Runtime/Launch/Private/Android/LaunchAndroid.cpp
--- a/Engine/Source/Runtime/Launch/Private/Android/LaunchAndroid.cpp
+++ b/Engine/Source/Runtime/Launch/Private/Android/LaunchAndroid.cpp
@@ -383,8 +383,6 @@ static void AndroidProcessEvents(struct android_app* state)
 	while((ident = ALooper_pollAll(-1, &fdesc, &events, (void**)&source)) >= 0)
 	// while((ident = ALooper_pollAll(0, &fdesc, &events, (void**)&source)) >= 0)
 	{
-		FPlatformMisc::LowLevelOutputDebugString(L"AndroidProcessEvents");
-
 		// process this event
 		if (source)
 		{
