Disable threaded rendering on Kindle Fire (1st gen) to fix black screen on startup.
#Android
#TTP 339151

diff --git a/Engine/Source/Runtime/Launch/Private/LaunchEngineLoop.cpp b/Engine/Source/Runtime/Launch/Private/LaunchEngineLoop.cpp
--- a/Engine/Source/Runtime/Launch/Private/LaunchEngineLoop.cpp
+++ b/Engine/Source/Runtime/Launch/Private/LaunchEngineLoop.cpp
@@ -1021,6 +1021,12 @@ int32 FEngineLoop::PreInit( const TCHAR* CmdLine )
 		{
 			GUseThreadedRendering = false;
 		}
+
+		// there is an issue with presenting the buffer on kindle fire (1st gen) with multiple threads using opengl
+		if (FAndroidMisc::GetDeviceModel() == FString(TEXT("Kindle Fire")))
+		{
+			GUseThreadedRendering = false;
+		}
 #endif
 	}
 #endif
