Ensure that EKeys are always initialized, not just when there is a slate application being created.

diff --git a/Engine/Source/Runtime/Launch/Launch.Build.cs b/Engine/Source/Runtime/Launch/Launch.Build.cs
--- a/Engine/Source/Runtime/Launch/Launch.Build.cs
+++ b/Engine/Source/Runtime/Launch/Launch.Build.cs
@@ -22,6 +22,7 @@ public class Launch : ModuleRules
 				"Core",
 				"CoreUObject",
 				"Engine",
+				"InputCore",
                 "MoviePlayer",
 				"Networking",
 				"PakFile",
diff --git a/Engine/Source/Runtime/Launch/Private/LaunchEngineLoop.cpp b/Engine/Source/Runtime/Launch/Private/LaunchEngineLoop.cpp
--- a/Engine/Source/Runtime/Launch/Private/LaunchEngineLoop.cpp
+++ b/Engine/Source/Runtime/Launch/Private/LaunchEngineLoop.cpp
@@ -1196,21 +1196,13 @@ int32 FEngineLoop::PreInit( const TCHAR* CmdLine )
 		FPlatformSplash::Show();
 	}
 
-	if ( !IsRunningDedicatedServer() )
-	{
-		if ( bHasEditorToken || bIsRegularClient )
-		{
-			// Init platform application
-			FSlateApplication::Create();
-		}
-		else
-		{
-			FCoreStyle::ResetToDefault();
-		}
-	}
-	else
+	EKeys::Initialize();
+	FCoreStyle::ResetToDefault();
+
+	if (!IsRunningDedicatedServer() && (bHasEditorToken || bIsRegularClient))
 	{
-		FCoreStyle::ResetToDefault();
+		// Init platform application
+		FSlateApplication::Create();
 	}
 
 	{
diff --git a/Engine/Source/Runtime/Slate/Private/Framework/Application/SlateApplication.cpp b/Engine/Source/Runtime/Slate/Private/Framework/Application/SlateApplication.cpp
--- a/Engine/Source/Runtime/Slate/Private/Framework/Application/SlateApplication.cpp
+++ b/Engine/Source/Runtime/Slate/Private/Framework/Application/SlateApplication.cpp
@@ -297,10 +297,6 @@ static bool LocateWidgetsUnderCursor_Helper( FArrangedWidget& Candidate, FVector
 
 void FSlateApplication::Create()
 {
-	EKeys::Initialize();
-
-	FCoreStyle::ResetToDefault();
-
 	CurrentApplication = MakeShareable( new FSlateApplication() );
 	CurrentBaseApplication = CurrentApplication;
 
