Better fix for CL# 2169902 that doesn't cause issues for standalone slate applications.

diff --git a/Engine/Source/Runtime/Launch/Private/LaunchEngineLoop.cpp b/Engine/Source/Runtime/Launch/Private/LaunchEngineLoop.cpp
--- a/Engine/Source/Runtime/Launch/Private/LaunchEngineLoop.cpp
+++ b/Engine/Source/Runtime/Launch/Private/LaunchEngineLoop.cpp
@@ -1196,14 +1196,18 @@ int32 FEngineLoop::PreInit( const TCHAR* CmdLine )
 		FPlatformSplash::Show();
 	}
 
-	EKeys::Initialize();
-	FCoreStyle::ResetToDefault();
-
 	if (!IsRunningDedicatedServer() && (bHasEditorToken || bIsRegularClient))
 	{
 		// Init platform application
 		FSlateApplication::Create();
 	}
+	else
+	{
+		// If we're not creating the slate application there is some basic initialization
+		// that it does that still must be done
+		EKeys::Initialize();
+		FCoreStyle::ResetToDefault();
+	}
 
 	{
 		DECLARE_SCOPE_CYCLE_COUNTER(TEXT("Initial UObject load"), STAT_InitialUObjectLoad, STATGROUP_LoadTime);
diff --git a/Engine/Source/Runtime/Slate/Private/Framework/Application/SlateApplication.cpp b/Engine/Source/Runtime/Slate/Private/Framework/Application/SlateApplication.cpp
--- a/Engine/Source/Runtime/Slate/Private/Framework/Application/SlateApplication.cpp
+++ b/Engine/Source/Runtime/Slate/Private/Framework/Application/SlateApplication.cpp
@@ -297,6 +297,10 @@ static bool LocateWidgetsUnderCursor_Helper( FArrangedWidget& Candidate, FVector
 
 void FSlateApplication::Create()
 {
+	EKeys::Initialize();
+
+	FCoreStyle::ResetToDefault();
+
 	CurrentApplication = MakeShareable( new FSlateApplication() );
 	CurrentBaseApplication = CurrentApplication;
 
