Fix for Samsung S3 Mini choppy rendering
#TTP 340563
#Android

diff --git a/Engine/Source/Runtime/Launch/Private/LaunchEngineLoop.cpp b/Engine/Source/Runtime/Launch/Private/LaunchEngineLoop.cpp
--- a/Engine/Source/Runtime/Launch/Private/LaunchEngineLoop.cpp
+++ b/Engine/Source/Runtime/Launch/Private/LaunchEngineLoop.cpp
@@ -1027,6 +1027,12 @@ int32 FEngineLoop::PreInit( const TCHAR* CmdLine )
 		{
 			GUseThreadedRendering = false;
 		}
+
+		// there is an issue with swapbuffer ordering on startup on samsung s3 mini with multiple threads using opengl
+		if (FAndroidMisc::GetDeviceModel() == FString(TEXT("GT-I8190L")))
+		{
+			GUseThreadedRendering = false;
+		}
 #endif
 	}
 #endif
