Fixed scalability settings not being applied on restart.

The settings are stored in game agnostic ini configs now. This fixes TTP340388 - CRITICAL: Mac: EDITOR: Automatic scalability settings are not applied after selecting the option to detect and apply changes (from the prompt).

diff --git a/Engine/Source/Runtime/Launch/Private/LaunchEngineLoop.cpp b/Engine/Source/Runtime/Launch/Private/LaunchEngineLoop.cpp
--- a/Engine/Source/Runtime/Launch/Private/LaunchEngineLoop.cpp
+++ b/Engine/Source/Runtime/Launch/Private/LaunchEngineLoop.cpp
@@ -1011,7 +1011,7 @@ int32 FEngineLoop::PreInit( const TCHAR* CmdLine )
 	// after SystemSettings.ini file loading so we get the right state,
 	// before ConsoleVariables.ini so the local developer can always override.
 	// before InitializeCVarsForActiveDeviceProfile() so the platform can override user settings
-	Scalability::LoadState(bHasEditorToken ? GEditorUserSettingsIni : GGameUserSettingsIni);
+	Scalability::LoadState((bHasEditorToken && !GEditorGameAgnosticIni.IsEmpty()) ? GEditorGameAgnosticIni : GGameUserSettingsIni);
 
 	// Set all CVars which have been setup in the device profiles.
 	UDeviceProfileManager::InitializeCVarsForActiveDeviceProfile();
