Fixed missing Android sensor events (tilt, gravity, rotation, acceleration).

diff --git a/Engine/Source/Runtime/Launch/Private/Android/LaunchAndroid.cpp b/Engine/Source/Runtime/Launch/Private/Android/LaunchAndroid.cpp
--- a/Engine/Source/Runtime/Launch/Private/Android/LaunchAndroid.cpp
+++ b/Engine/Source/Runtime/Launch/Private/Android/LaunchAndroid.cpp
@@ -84,6 +84,7 @@
 //Android event callback functions
 static int32_t HandleInputCB(struct android_app* app, AInputEvent* event); //Touch and key input events
 static void OnAppCommandCB(struct android_app* app, int32_t cmd); //Lifetime events
+static int HandleSensorEvents(int fd, int events, void* data); // Sensor events
 
 bool GHasInterruptionRequest = false;
 bool GIsInterrupted = false;
@@ -363,7 +364,7 @@
 			SensorManager, ASENSOR_TYPE_GYROSCOPE);
 		// Create the queue for events to arrive.
 		SensorQueue = ASensorManager_createEventQueue(
-			SensorManager, state->looper, LOOPER_ID_USER, NULL, NULL);
+			SensorManager, state->looper, LOOPER_ID_USER, HandleSensorEvents, NULL);
 	}
 
 	FPlatformMisc::LowLevelOutputDebugString(L"Passed sensor initialization");
@@ -391,16 +392,6 @@
 	int events;
 	struct android_poll_source* source;
 
-	// It's not possible to discern sequencing across sensors in
-	// Android. So we average out all the sensor events on one cycle
-	// and post a single motion sensor data point. We also need
-	// to synthesize additional information.
-	FVector current_accelerometer(0, 0, 0);
-	FVector current_gyroscope(0, 0, 0);
-	int32 current_accelerometer_sample_count = 0;
-	int32 current_gyroscope_sample_count = 0;
-	static FVector last_accelerometer(0, 0, 0);
-
 	while((ident = ALooper_pollAll(-1, &fdesc, &events, (void**)&source)) >= 0)
 	{
 		// process this event
@@ -408,116 +399,6 @@
 		{
 			source->process(state, source);
 		}
-
-		// process sensor events
-		if (ident == LOOPER_ID_USER)
-		{
-			if (NULL != SensorAccelerometer || NULL != SensorGyroscope)
-			{
-				ASensorEvent sensor_event;
-				while (ASensorEventQueue_getEvents(SensorQueue, &sensor_event, 1) > 0)
-				{
-					if (ASENSOR_TYPE_ACCELEROMETER == sensor_event.type)
-					{
-						current_accelerometer.X += sensor_event.acceleration.x;
-						current_accelerometer.Y += sensor_event.acceleration.y;
-						current_accelerometer.Z += sensor_event.acceleration.z;
-						current_accelerometer_sample_count += 1;
-					}
-					else if (ASENSOR_TYPE_GYROSCOPE == sensor_event.type)
-					{
-						current_gyroscope.X += sensor_event.vector.pitch;
-						current_gyroscope.Y += sensor_event.vector.azimuth;
-						current_gyroscope.Z += sensor_event.vector.roll;
-						current_gyroscope_sample_count += 1;
-					}
-				}
-			}
-		}
 	}
-
-	if (current_accelerometer_sample_count > 0)
-	{
-		// Do simple average of the samples we just got.
-		current_accelerometer /= float(current_accelerometer_sample_count);
-		last_accelerometer = current_accelerometer;
-	}
-	else
-	{
-		current_accelerometer = last_accelerometer;
-	}
-
-	if (current_gyroscope_sample_count > 0)
-	{
-		// Do simple average of the samples we just got.
-		current_gyroscope /= float(current_gyroscope_sample_count);
-	}
-
-	// If we have motion samples we generate the single event.
-	if (current_accelerometer_sample_count > 0 ||
-		current_gyroscope_sample_count > 0)
-	{
-		// The data we compose the motion event from.
-		FVector current_tilt(0, 0, 0);
-		FVector current_rotation_rate(0, 0, 0);
-		FVector current_gravity(0, 0, 0);
-		FVector current_acceleration(0, 0, 0);
-
-		// Buffered, historical, motion data.
-		static FVector last_tilt(0, 0, 0);
-		static FVector last_gravity(0, 0, 0);
-
-		// We use a low-pass filter to synthesize the gravity
-		// vector.
-		static bool first_acceleration_sample = true;
-		if (!first_acceleration_sample)
-		{
-			current_gravity
-				= last_gravity*SampleDecayRate
-				+ current_accelerometer*(1.0f - SampleDecayRate);
-		}
-		first_acceleration_sample = false;
-
-		// Calc the tilt from the accelerometer as it's not
-		// available directly.
-		FVector accelerometer_dir = -current_accelerometer.SafeNormal();
-		float current_pitch
-			= FMath::Atan2(accelerometer_dir.Y, accelerometer_dir.Z);
-		float current_roll
-			= -FMath::Atan2(accelerometer_dir.X, accelerometer_dir.Z);
-		current_tilt.X = current_pitch;
-		current_tilt.Y = 0;
-		current_tilt.Z = current_roll;
-
-		// And take out the gravity from the accel to get
-		// the linear acceleration.
-		current_acceleration = current_accelerometer - current_gravity;
-
-		if (current_gyroscope_sample_count > 0)
-		{
-			// The rotation rate is the what the gyroscope gives us.
-			current_rotation_rate = current_gyroscope;
-		}
-		else if (NULL == SensorGyroscope)
-		{
-			// If we don't have a gyroscope at all we need to calc a rotation
-			// rate from our calculated tilt and a delta.
-			current_rotation_rate = current_tilt - last_tilt;
-		}
-
-		// Finally record the motion event with all the data.
-		FAndroidInputInterface::QueueMotionData(current_tilt,
-			current_rotation_rate, current_gravity, current_acceleration);
-
-		// Update history values.
-		last_tilt = current_tilt;
-		last_gravity = current_gravity;
-
-		// UE_LOG(LogTemp, Log, TEXT("MOTION: tilt = %s, rotation-rate = %s, gravity = %s, acceleration = %s"),
-		//	*current_tilt.ToCompactString(),
-		//	*current_rotation_rate.ToCompactString(),
-		//	*current_gravity.ToCompactString(),
-		//	*current_acceleration.ToCompactString());
-	}
 }
 
@@ -894,6 +775,128 @@
 		EventHandlerEvent->Trigger();
 }
 
+static int HandleSensorEvents(int fd, int events, void* data)
+{
+	// It's not possible to discern sequencing across sensors in
+	// Android. So we average out all the sensor events on one cycle
+	// and post a single motion sensor data point. We also need
+	// to synthesize additional information.
+	FVector current_accelerometer(0, 0, 0);
+	FVector current_gyroscope(0, 0, 0);
+	int32 current_accelerometer_sample_count = 0;
+	int32 current_gyroscope_sample_count = 0;
+	static FVector last_accelerometer(0, 0, 0);
+
+	if (NULL != SensorAccelerometer || NULL != SensorGyroscope)
+	{
+		ASensorEvent sensor_event;
+		while (ASensorEventQueue_getEvents(SensorQueue, &sensor_event, 1) > 0)
+		{
+			if (ASENSOR_TYPE_ACCELEROMETER == sensor_event.type)
+			{
+				current_accelerometer.X += sensor_event.acceleration.x;
+				current_accelerometer.Y += sensor_event.acceleration.y;
+				current_accelerometer.Z += sensor_event.acceleration.z;
+				current_accelerometer_sample_count += 1;
+			}
+			else if (ASENSOR_TYPE_GYROSCOPE == sensor_event.type)
+			{
+				current_gyroscope.X += sensor_event.vector.pitch;
+				current_gyroscope.Y += sensor_event.vector.azimuth;
+				current_gyroscope.Z += sensor_event.vector.roll;
+				current_gyroscope_sample_count += 1;
+			}
+		}
+	}
+
+	if (current_accelerometer_sample_count > 0)
+	{
+		// Do simple average of the samples we just got.
+		current_accelerometer /= float(current_accelerometer_sample_count);
+		last_accelerometer = current_accelerometer;
+	}
+	else
+	{
+		current_accelerometer = last_accelerometer;
+	}
+
+	if (current_gyroscope_sample_count > 0)
+	{
+		// Do simple average of the samples we just got.
+		current_gyroscope /= float(current_gyroscope_sample_count);
+	}
+
+	// If we have motion samples we generate the single event.
+	if (current_accelerometer_sample_count > 0 ||
+		current_gyroscope_sample_count > 0)
+	{
+		// The data we compose the motion event from.
+		FVector current_tilt(0, 0, 0);
+		FVector current_rotation_rate(0, 0, 0);
+		FVector current_gravity(0, 0, 0);
+		FVector current_acceleration(0, 0, 0);
+
+		// Buffered, historical, motion data.
+		static FVector last_tilt(0, 0, 0);
+		static FVector last_gravity(0, 0, 0);
+
+		// We use a low-pass filter to synthesize the gravity
+		// vector.
+		static bool first_acceleration_sample = true;
+		if (!first_acceleration_sample)
+		{
+			current_gravity
+				= last_gravity*SampleDecayRate
+				+ current_accelerometer*(1.0f - SampleDecayRate);
+		}
+		first_acceleration_sample = false;
+
+		// Calc the tilt from the accelerometer as it's not
+		// available directly.
+		FVector accelerometer_dir = -current_accelerometer.SafeNormal();
+		float current_pitch
+			= FMath::Atan2(accelerometer_dir.Y, accelerometer_dir.Z);
+		float current_roll
+			= -FMath::Atan2(accelerometer_dir.X, accelerometer_dir.Z);
+		current_tilt.X = current_pitch;
+		current_tilt.Y = 0;
+		current_tilt.Z = current_roll;
+
+		// And take out the gravity from the accel to get
+		// the linear acceleration.
+		current_acceleration = current_accelerometer - current_gravity;
+
+		if (current_gyroscope_sample_count > 0)
+		{
+			// The rotation rate is the what the gyroscope gives us.
+			current_rotation_rate = current_gyroscope;
+		}
+		else if (NULL == SensorGyroscope)
+		{
+			// If we don't have a gyroscope at all we need to calc a rotation
+			// rate from our calculated tilt and a delta.
+			current_rotation_rate = current_tilt - last_tilt;
+		}
+
+		// Finally record the motion event with all the data.
+		FAndroidInputInterface::QueueMotionData(current_tilt,
+			current_rotation_rate, current_gravity, current_acceleration);
+
+		// Update history values.
+		last_tilt = current_tilt;
+		last_gravity = current_gravity;
+
+		// UE_LOG(LogTemp, Log, TEXT("MOTION: tilt = %s, rotation-rate = %s, gravity = %s, acceleration = %s"),
+		//	*current_tilt.ToCompactString(),
+		//	*current_rotation_rate.ToCompactString(),
+		//	*current_gravity.ToCompactString(),
+		//	*current_acceleration.ToCompactString());
+	}
+
+	// Indicate we want to keep getting events.
+	return 1;
+}
+
 //Native-defined functions
 
 //This function is declared in the Java-defined class, GameActivity.java: "public native void nativeConsoleCommand(String commandString);"
