Refactoring CoreUObject initialization to make LaunchEngineLoop not depend directly on CoreUObject module code.

#change Moved static class initalization to CoreUObject module (StartupModule)
#change Re-ordered PreInit initialization
#change Fixed a crash (when using dashes from Outlook) and duplicated code in commandline initialization

diff --git a/Engine/Source/Runtime/Core/Private/Misc/CoreMisc.cpp b/Engine/Source/Runtime/Core/Private/Misc/CoreMisc.cpp
--- a/Engine/Source/Runtime/Core/Private/Misc/CoreMisc.cpp
+++ b/Engine/Source/Runtime/Core/Private/Misc/CoreMisc.cpp
@@ -558,17 +558,24 @@ const TCHAR* FCommandLine::Get()
 	return CmdLine;
 }
 
-void FCommandLine::Set(const TCHAR* NewCommandLine)
-{	
+bool FCommandLine::Set(const TCHAR* NewCommandLine)
+{		
 	FCString::Strncpy( CmdLine, NewCommandLine, ARRAY_COUNT(CmdLine) );
+	bIsInitialized = true;
 
-	// Detect dashes (0x2013) and report an error if they're found
+	// Check for the '-' that normal ones get converted to in Outlook. It's important to do it AFTER the command line is initialized
 	if (StringHasBadDashes(NewCommandLine))
 	{
-		UE_LOG(LogInit, Fatal, TEXT("Illegal character detected in the command line. Replace dash (0x2013) with a hyphen."));
+		FText ErrorMessage = FText::Format(NSLOCTEXT("Engine", "ComdLineHasInvalidChar", "Error: Command-line contains an invalid '-' character, likely pasted from an email.\nCmdline = {0}"), FText::FromString(NewCommandLine));
+#if !UE_BUILD_SHIPPING
+		FMessageDialog::Open(EAppMsgType::Ok, ErrorMessage);
+		return false;
+#else
+		UE_LOG(LogInit, Fatal, TEXT("%s"), *ErrorMessage.ToString());
+#endif
 	}
 
-	bIsInitialized = true;
+	return true;
 }
 
 void FCommandLine::Append(const TCHAR* AppendString)
diff --git a/Engine/Source/Runtime/Core/Public/Misc/CoreMisc.h b/Engine/Source/Runtime/Core/Public/Misc/CoreMisc.h
--- a/Engine/Source/Runtime/Core/Public/Misc/CoreMisc.h
+++ b/Engine/Source/Runtime/Core/Public/Misc/CoreMisc.h
@@ -82,7 +82,7 @@ struct CORE_API FCommandLine
 	/**
 	 * Sets CmdLine to the string given
 	 */
-	static void Set(const TCHAR* NewCommandLine);
+	static bool Set(const TCHAR* NewCommandLine);
 
 	/**
 	 * Appends the passed string to the command line as it is (no space is added).
diff --git a/Engine/Source/Runtime/CoreUObject/Private/UObject/CoreNative.cpp b/Engine/Source/Runtime/CoreUObject/Private/UObject/CoreNative.cpp
--- a/Engine/Source/Runtime/CoreUObject/Private/UObject/CoreNative.cpp
+++ b/Engine/Source/Runtime/CoreUObject/Private/UObject/CoreNative.cpp
@@ -11,6 +11,10 @@ class FCoreUObjectModule : public FDefaultModuleImpl
 public:
 	virtual void StartupModule() override
 	{
+		// Register all classes that have been loaded so far. This is required for CVars to work.
+		void UClassRegisterAllCompiledInClasses();
+		UClassRegisterAllCompiledInClasses();
+
 		void InitUObject();
 		FCoreDelegates::OnInit.AddStatic(InitUObject);
 
diff --git a/Engine/Source/Runtime/CoreUObject/Public/UObject/UObjectBase.h b/Engine/Source/Runtime/CoreUObject/Public/UObject/UObjectBase.h
--- a/Engine/Source/Runtime/CoreUObject/Public/UObject/UObjectBase.h
+++ b/Engine/Source/Runtime/CoreUObject/Public/UObject/UObjectBase.h
@@ -248,11 +248,6 @@ FORCEINLINE bool UObjectInitialized() { return Internal::GObjInitialized; }
  */
 COREUOBJECT_API void UObjectForceRegistration(UObjectBase* Object);
 
-/**
- * Registers all recently loaded native classes
- */
-COREUOBJECT_API void UClassRegisterAllCompiledInClasses();
-
 /** 
  * Base class for deferred native class registration
  */
diff --git a/Engine/Source/Runtime/Launch/Private/LaunchEngineLoop.cpp b/Engine/Source/Runtime/Launch/Private/LaunchEngineLoop.cpp
--- a/Engine/Source/Runtime/Launch/Private/LaunchEngineLoop.cpp
+++ b/Engine/Source/Runtime/Launch/Private/LaunchEngineLoop.cpp
@@ -86,10 +86,6 @@
 	#include "HideWindowsPlatformTypes.h"
 #endif
 
-#if WITH_COREUOBJECT
-	#include "CoreUObject.h"
-#endif
-
 // Pipe output to std output
 // This enables UBT to collect the output for it's own use
 class FOutputDeviceStdOutput : public FOutputDevice
@@ -622,36 +618,29 @@ int32 FEngineLoop::PreInit( const TCHAR* CmdLine )
 
 	// Switch into executable's directory.
 	FPlatformProcess::SetCurrentWorkingDirectoryToBaseDir();
-#if WITH_COREUOBJECT
-	// Register all classes that have been loaded so far. This is required for CVars to work.
-	UClassRegisterAllCompiledInClasses();
-#endif
-
-#if WITH_ENGINE
-	extern ENGINE_API void InitializeRenderingCVarsCaching();
-	InitializeRenderingCVarsCaching();
-#endif
 
-	// Initialize log console here to avoid statics initialization issues when launched from the command line.
-	GScopedLogConsole = FPlatformOutputDevices::GetLogConsole();
-
-#if !UE_BUILD_SHIPPING
-	// Check for the '-' that normal ones get converted to in Outlook
-	if (StringHasBadDashes(CmdLine))
+	// this is set later with shorter command lines, but we want to make sure it is set ASAP as some subsystems will do the tests themselves...
+	// also realize that command lines can be pulled from the network at a slightly later time.
+	if (!FCommandLine::Set(CmdLine))
 	{
-		FMessageDialog::Open( EAppMsgType::Ok, FText::Format( NSLOCTEXT("Engine", "ComdLineHasInvalidChar", "Error: Command-line contains an invalid '-' character, likely pasted from an email.\nCmdline = {0}"), FText::FromString( CmdLine ) ) );
+		// Fail, shipping builds will crash if setting command line fails
 		return -1;
 	}
-#endif
+
+	// Set GGameName, based on the command line
+	if (LaunchSetGameName(CmdLine) == false)
+	{
+		// If it failed, do not continue
+		return 1;
+	}
+
+	// Initialize log console here to avoid statics initialization issues when launched from the command line.
+	GScopedLogConsole = FPlatformOutputDevices::GetLogConsole();
 
 	// Always enable the backlog so we get all messages, we will disable and clear it in the game
 	// as soon as we determine whether GIsEditor == false
 	GLog->EnableBacklog(true);
 
-	// this is set later with shorter command lines, but we want to make sure it is set ASAP as some subsystems will do the tests themselves...
-	// also realize that command lines can be pulled from the network at a slightly later time.
-	FCommandLine::Set(CmdLine); 
-
 	// Initialize std out device as early as possible if requested in the command line
 	if (FParse::Param(FCommandLine::Get(), TEXT("stdout")))
 	{
@@ -674,13 +663,6 @@ int32 FEngineLoop::PreInit( const TCHAR* CmdLine )
 	}
 #endif // !UE_BUILD_SHIPPING
 
-	// Set GGameName, based on the command line
-	if (LaunchSetGameName(CmdLine) == false)
-	{
-		// If it failed, do not continue
-		return 1;
-	}
-
 	// Switch into executable's directory (may be required by some of the platform file overrides)
 	FPlatformProcess::SetCurrentWorkingDirectoryToBaseDir();
 
@@ -729,6 +711,14 @@ int32 FEngineLoop::PreInit( const TCHAR* CmdLine )
 		}
 	}
 
+	// Load Core modules required for everything else to work (needs to be loaded before InitializeRenderingCVarsCaching)
+	LoadCoreModules();
+
+#if WITH_ENGINE
+	extern ENGINE_API void InitializeRenderingCVarsCaching();
+	InitializeRenderingCVarsCaching();
+#endif
+
 	// remember thread id of the main thread
 	GGameThreadId = FPlatformTLS::GetCurrentThreadId();
 	GIsGameThreadIdInitialized = true;
@@ -1630,18 +1620,19 @@ int32 FEngineLoop::PreInit( const TCHAR* CmdLine )
 	return 0;
 }
 
+void FEngineLoop::LoadCoreModules()
+{
+	// Always attempt to load CoreUObject. It requires additional pre-init which is called from its module's StartupModule method.
+#if WITH_COREUOBJECT
+	FModuleManager::Get().LoadModule(TEXT("CoreUObject"));
+#endif
+}
 
 void FEngineLoop::LoadPreInitModules()
 {	
 	DECLARE_SCOPE_CYCLE_COUNTER(TEXT("Loading PreInit Modules"), STAT_PreInitModules, STATGROUP_LoadTime);
 
 	// GGetMapNameDelegate is initialized here
-
-	// Always attempt to load CoreUObject. It requires additional pre init which is called from its module's StartupModule method.
-#if WITH_COREUOBJECT
-	FModuleManager::Get().LoadModule(TEXT("CoreUObject"));
-#endif
-
 #if WITH_ENGINE
 	FModuleManager::Get().LoadModule(TEXT("Engine"));
 
diff --git a/Engine/Source/Runtime/Launch/Public/LaunchEngineLoop.h b/Engine/Source/Runtime/Launch/Public/LaunchEngineLoop.h
--- a/Engine/Source/Runtime/Launch/Public/LaunchEngineLoop.h
+++ b/Engine/Source/Runtime/Launch/Public/LaunchEngineLoop.h
@@ -53,6 +53,11 @@ public:
 	 */ 
 	void LoadPreInitModules();
 
+	/**
+	* Load core modules.
+	*/
+	void LoadCoreModules();
+
 #if WITH_ENGINE
 	
 	/**
