Added code to move android application to background (minimize). Also implemented quit for android.
Added JNI error codes to make errors in that easier to debug.

diff --git a/Engine/Build/Android/Java/src/com/epicgames/ue4/GameActivity.java b/Engine/Build/Android/Java/src/com/epicgames/ue4/GameActivity.java
--- a/Engine/Build/Android/Java/src/com/epicgames/ue4/GameActivity.java
+++ b/Engine/Build/Android/Java/src/com/epicgames/ue4/GameActivity.java
@@ -1034,7 +1034,19 @@ public class GameActivity extends NativeActivity implements GoogleApiClient.Conn
 		return JavaBuildSettings.PackageType.AMAZON == JavaBuildSettings.PACKAGING;
 	}
 
+	public void AndroidThunkJava_Minimize()
+	{
+		Intent startMain = new Intent(Intent.ACTION_MAIN);
+		startMain.addCategory(Intent.CATEGORY_HOME);
+		startMain.setFlags(Intent.FLAG_ACTIVITY_NEW_TASK);
+		startActivity(startMain);
+	}
 
+	public void AndroidThunkJava_ForceQuit()
+	{
+		System.exit(0);
+		// finish();
+	}
 
 	public native boolean nativeIsShippingBuild();
 	public native void nativeSetGlobalActivity();
diff --git a/Engine/Source/Runtime/Core/Private/Android/AndroidMisc.cpp b/Engine/Source/Runtime/Core/Private/Android/AndroidMisc.cpp
--- a/Engine/Source/Runtime/Core/Private/Android/AndroidMisc.cpp
+++ b/Engine/Source/Runtime/Core/Private/Android/AndroidMisc.cpp
@@ -22,6 +22,21 @@ GenericApplication* FAndroidMisc::CreateApplication()
 	return FAndroidApplication::CreateAndroidApplication();
 }
 
+extern void AndroidThunkCpp_ForceQuit();
+
+void FAndroidMisc::RequestExit( bool Force )
+{
+	AndroidThunkCpp_ForceQuit();
+}
+
+extern void AndroidThunkCpp_Minimize();
+
+void FAndroidMisc::RequestMinimize()
+{
+	AndroidThunkCpp_Minimize();
+}
+
+
 void FAndroidMisc::LowLevelOutputDebugString(const TCHAR *Message)
 {
 	LocalPrint(Message);
diff --git a/Engine/Source/Runtime/Core/Private/GenericPlatform/GenericPlatformMisc.cpp b/Engine/Source/Runtime/Core/Private/GenericPlatform/GenericPlatformMisc.cpp
--- a/Engine/Source/Runtime/Core/Private/GenericPlatform/GenericPlatformMisc.cpp
+++ b/Engine/Source/Runtime/Core/Private/GenericPlatform/GenericPlatformMisc.cpp
@@ -323,6 +323,10 @@ void FGenericPlatformMisc::LocalPrint( const TCHAR* Str )
 #endif
 }
 
+void FGenericPlatformMisc::RequestMinimize()
+{
+}
+
 void FGenericPlatformMisc::RequestExit( bool Force )
 {
 	UE_LOG(LogGenericPlatformMisc, Log,  TEXT("FPlatformMisc::RequestExit(%i)"), Force );
diff --git a/Engine/Source/Runtime/Core/Public/Android/AndroidMisc.h b/Engine/Source/Runtime/Core/Public/Android/AndroidMisc.h
--- a/Engine/Source/Runtime/Core/Public/Android/AndroidMisc.h
+++ b/Engine/Source/Runtime/Core/Public/Android/AndroidMisc.h
@@ -53,6 +53,8 @@ struct CORE_API FAndroidMisc : public FGenericPlatformMisc
 {
 	static class GenericApplication* CreateApplication();
 
+	static void RequestMinimize();
+	static void RequestExit( bool Force );
 	static void LowLevelOutputDebugString(const TCHAR *Message);
 	static void LocalPrint(const TCHAR *Message);
 	static void PlatformPreInit();
@@ -84,6 +86,7 @@ struct CORE_API FAndroidMisc : public FGenericPlatformMisc
 	static FString GetGLVersion();
 	static bool SupportsFloatingPointRenderTargets();
 
+
 #if !UE_BUILD_SHIPPING
 	FORCEINLINE static bool IsDebuggerPresent()
 	{
diff --git a/Engine/Source/Runtime/Core/Public/Android/AndroidProperties.h b/Engine/Source/Runtime/Core/Public/Android/AndroidProperties.h
--- a/Engine/Source/Runtime/Core/Public/Android/AndroidProperties.h
+++ b/Engine/Source/Runtime/Core/Public/Android/AndroidProperties.h
@@ -88,6 +88,11 @@ struct FAndroidPlatformProperties
 		return false;
 	}
 
+	static FORCEINLINE bool SupportsMinimize()
+	{
+		return true;
+	}
+
 	static FORCEINLINE bool SupportsQuit() 
 	{
 		return true;
diff --git a/Engine/Source/Runtime/Core/Public/GenericPlatform/GenericPlatformMisc.h b/Engine/Source/Runtime/Core/Public/GenericPlatform/GenericPlatformMisc.h
--- a/Engine/Source/Runtime/Core/Public/GenericPlatform/GenericPlatformMisc.h
+++ b/Engine/Source/Runtime/Core/Public/GenericPlatform/GenericPlatformMisc.h
@@ -479,6 +479,13 @@ public:
 	/** Prints string to the default output */
 	static void LocalPrint( const TCHAR* Str );
 
+
+	/**
+	 * Request application to minimize (goto background)
+	 *
+	 **/
+	static void RequestMinimize();
+
 	/**
 	 * Requests application exit.
 	 *
diff --git a/Engine/Source/Runtime/Core/Public/GenericPlatform/GenericPlatformProperties.h b/Engine/Source/Runtime/Core/Public/GenericPlatform/GenericPlatformProperties.h
--- a/Engine/Source/Runtime/Core/Public/GenericPlatform/GenericPlatformProperties.h
+++ b/Engine/Source/Runtime/Core/Public/GenericPlatform/GenericPlatformProperties.h
@@ -260,6 +260,11 @@ struct FGenericPlatformProperties
 		return false;
 	}
 
+	static FORCEINLINE bool SupportsMinimize()
+	{
+		return false;
+	}
+
 	// Whether the platform allows an application to quit to the OS
 	static FORCEINLINE bool SupportsQuit()
 	{
diff --git a/Engine/Source/Runtime/Engine/Classes/Engine/Engine.h b/Engine/Source/Runtime/Engine/Classes/Engine/Engine.h
--- a/Engine/Source/Runtime/Engine/Classes/Engine/Engine.h
+++ b/Engine/Source/Runtime/Engine/Classes/Engine/Engine.h
@@ -1606,7 +1606,6 @@ public:
 	 * Exec command handlers
 	 */
 	bool HandleFlushLogCommand( const TCHAR* Cmd, FOutputDevice& Ar );
-	bool HandleExitCommand( const TCHAR* Cmd, FOutputDevice& Ar );
 	bool HandleGameVerCommand( const TCHAR* Cmd, FOutputDevice& Ar );
 	bool HandleStatCommand( UWorld* World, FCommonViewportClient* ViewportClient, const TCHAR* Cmd, FOutputDevice& Ar );
 	bool HandleStartMovieCaptureCommand( const TCHAR* Cmd, FOutputDevice& Ar );
diff --git a/Engine/Source/Runtime/Engine/Classes/Engine/GameEngine.h b/Engine/Source/Runtime/Engine/Classes/Engine/GameEngine.h
--- a/Engine/Source/Runtime/Engine/Classes/Engine/GameEngine.h
+++ b/Engine/Source/Runtime/Engine/Classes/Engine/GameEngine.h
@@ -106,6 +106,7 @@ public:
 	bool HandleReattachComponentsCommand( const TCHAR* Cmd, FOutputDevice& Ar );
 	bool HandleCommand( const TCHAR* Cmd, FOutputDevice& Ar );
 	bool HandleExitCommand( const TCHAR* Cmd, FOutputDevice& Ar );
+	bool HandleMinimizeCommand( const TCHAR *Cmd, FOutputDevice &Ar );
 	bool HandleGetMaxTickRateCommand( const TCHAR* Cmd, FOutputDevice& Ar );
 	bool HandleCancelCommand( const TCHAR* Cmd, FOutputDevice& Ar, UWorld* InWorld );
 #if !UE_BUILD_SHIPPING
diff --git a/Engine/Source/Runtime/Engine/Classes/Kismet/KismetSystemLibrary.h b/Engine/Source/Runtime/Engine/Classes/Kismet/KismetSystemLibrary.h
--- a/Engine/Source/Runtime/Engine/Classes/Kismet/KismetSystemLibrary.h
+++ b/Engine/Source/Runtime/Engine/Classes/Kismet/KismetSystemLibrary.h
@@ -35,6 +35,18 @@ namespace EMoveComponentAction
 	};
 }
 
+UENUM()
+namespace EQuitPreference
+{
+	enum Type
+	{
+		/** Move the application to the background */
+		Background,
+		/** Exit the game completely */
+		Quit,
+	};
+}
+
 USTRUCT()
 struct FGenericStruct
 {
@@ -203,7 +215,7 @@ class UKismetSystemLibrary : public UBlueprintFunctionLibrary
 	 * @param	SpecificPlayer	The specific player to quit the game. If not specified, player 0 will quit.
 	 */
 	UFUNCTION(BlueprintCallable, Category="Game",meta=(HidePin="WorldContextObject", DefaultToSelf="WorldContextObject"))
-	static void QuitGame(UObject* WorldContextObject, class APlayerController* SpecificPlayer);
+	static void QuitGame(UObject* WorldContextObject, class APlayerController* SpecificPlayer, TEnumAsByte<EQuitPreference::Type> QuitPreference);
 
 	//=============================================================================
 	// Latent Actions
diff --git a/Engine/Source/Runtime/Engine/Private/GameEngine.cpp b/Engine/Source/Runtime/Engine/Private/GameEngine.cpp
--- a/Engine/Source/Runtime/Engine/Private/GameEngine.cpp
+++ b/Engine/Source/Runtime/Engine/Private/GameEngine.cpp
@@ -545,10 +545,20 @@ bool UGameEngine::Exec( UWorld* InWorld, const TCHAR* Cmd, FOutputDevice& Ar )
 	{
 		return HandleReattachComponentsCommand( Cmd, Ar );
 	}
-	// exec to 
 	else if( FParse::Command( &Cmd,TEXT("EXIT")) || FParse::Command(&Cmd,TEXT("QUIT")))
 	{
-		if ( FPlatformProperties::SupportsQuit() )
+		FString CmdName = FParse::Token(Cmd, 0);
+		bool Background = false;
+		if (!CmdName.IsEmpty() && !FCString::Stricmp(*CmdName, TEXT("background")))
+		{
+			Background = true;
+		}
+
+		if ( Background && FPlatformProperties::SupportsMinimize() )
+		{
+			return HandleMinimizeCommand( Cmd, Ar );
+		}
+		else if ( FPlatformProperties::SupportsQuit() )
 		{
 			return HandleExitCommand( Cmd, Ar );
 		}
@@ -647,6 +657,14 @@ bool UGameEngine::HandleExitCommand( const TCHAR* Cmd, FOutputDevice& Ar )
 	return true;
 }
 
+bool UGameEngine::HandleMinimizeCommand( const TCHAR *Cmd, FOutputDevice &Ar )
+{
+	Ar.Log( TEXT("Minimize by request") );
+	FPlatformMisc::RequestMinimize();
+
+	return true;
+}
+
 bool UGameEngine::HandleGetMaxTickRateCommand( const TCHAR* Cmd, FOutputDevice& Ar )
 {
 	Ar.Logf( TEXT("%f"), GetMaxTickRate(0,false) );
diff --git a/Engine/Source/Runtime/Engine/Private/KismetSystemLibrary.cpp b/Engine/Source/Runtime/Engine/Private/KismetSystemLibrary.cpp
--- a/Engine/Source/Runtime/Engine/Private/KismetSystemLibrary.cpp
+++ b/Engine/Source/Runtime/Engine/Private/KismetSystemLibrary.cpp
@@ -220,12 +220,19 @@ void UKismetSystemLibrary::ExecuteConsoleCommand(UObject* WorldContextObject, co
 	}
 }
 
-void UKismetSystemLibrary::QuitGame(UObject* WorldContextObject, class APlayerController* SpecificPlayer)
+void UKismetSystemLibrary::QuitGame(UObject* WorldContextObject, class APlayerController* SpecificPlayer, TEnumAsByte<EQuitPreference::Type> QuitPreference)
 {
 	APlayerController* TargetPC = SpecificPlayer ? SpecificPlayer : UGameplayStatics::GetPlayerController(WorldContextObject, 0);
 	if( TargetPC )
 	{
-		TargetPC->ConsoleCommand("quit");
+		if ( QuitPreference == EQuitPreference::Background)
+		{
+			TargetPC->ConsoleCommand("quit background");
+		}
+		else
+		{
+			TargetPC->ConsoleCommand("quit");
+		}
 	}
 }
 
diff --git a/Engine/Source/Runtime/Engine/Private/UnrealEngine.cpp b/Engine/Source/Runtime/Engine/Private/UnrealEngine.cpp
--- a/Engine/Source/Runtime/Engine/Private/UnrealEngine.cpp
+++ b/Engine/Source/Runtime/Engine/Private/UnrealEngine.cpp
@@ -2123,10 +2123,6 @@ bool UEngine::Exec( UWorld* InWorld, const TCHAR* Cmd, FOutputDevice& Ar )
 	{
 		return HandleFlushLogCommand( Cmd, Ar );
 	}
-	else if( FParse::Command(&Cmd,TEXT("EXIT")) || FParse::Command(&Cmd,TEXT("QUIT")))
-	{
-		return HandleExitCommand( Cmd, Ar );
-	}
 	else if( FParse::Command(&Cmd, TEXT("GAMEVER")) ||  FParse::Command(&Cmd, TEXT("GAMEVERSION")))
 	{
 		return HandleGameVerCommand( Cmd, Ar );
@@ -5132,17 +5128,6 @@ bool UEngine::HandleFlushLogCommand( const TCHAR* Cmd, FOutputDevice& Ar )
 	return true;
 }
 
-bool UEngine::HandleExitCommand( const TCHAR* Cmd, FOutputDevice& Ar )
-{
-	// Ignore these commands when running the editor
-	if( !GIsEditor )
-	{
-		Ar.Log( TEXT("Closing by request") );
-		FPlatformMisc::RequestExit( 0 );
-	}
-	return true;
-}
-
 bool UEngine::HandleDumpTicksCommand( UWorld* InWorld, const TCHAR* Cmd, FOutputDevice& Ar )
 {
 	// Handle optional parameters, will dump all tick functions by default.
diff --git a/Engine/Source/Runtime/Launch/Private/Android/AndroidJNI.cpp b/Engine/Source/Runtime/Launch/Private/Android/AndroidJNI.cpp
--- a/Engine/Source/Runtime/Launch/Private/Android/AndroidJNI.cpp
+++ b/Engine/Source/Runtime/Launch/Private/Android/AndroidJNI.cpp
@@ -118,6 +118,8 @@ jmethodID JDef_GameActivity::AndroidThunkJava_ShowAdBanner;
 jmethodID JDef_GameActivity::AndroidThunkJava_HideAdBanner;
 jmethodID JDef_GameActivity::AndroidThunkJava_CloseAdBanner;
 jmethodID JDef_GameActivity::AndroidThunkJava_GetAssetManager;
+jmethodID JDef_GameActivity::AndroidThunkJava_Minimize;
+jmethodID JDef_GameActivity::AndroidThunkJava_ForceQuit;
 
 jclass JDef_GameActivity::JavaAchievementClassID;
 jfieldID JDef_GameActivity::AchievementIDField;
@@ -336,6 +338,22 @@ AAssetManager * AndroidThunkCpp_GetAssetManager()
 	return GAssetManagerRef;
 }
 
+void AndroidThunkCpp_Minimize()
+{
+	if (JNIEnv* Env = GetJavaEnv())
+	{
+		Env->CallVoidMethod(GJavaGlobalThis, JDef_GameActivity::AndroidThunkJava_Minimize);
+	}
+}
+
+void AndroidThunkCpp_ForceQuit()
+{
+	if (JNIEnv* Env = GetJavaEnv())
+	{
+		Env->CallVoidMethod(GJavaGlobalThis, JDef_GameActivity::AndroidThunkJava_ForceQuit);
+	}
+}
+
 //The JNI_OnLoad function is triggered by loading the game library from 
 //the Java source file.
 //	static
@@ -346,6 +364,17 @@ AAssetManager * AndroidThunkCpp_GetAssetManager()
 // Use the JNI_OnLoad function to map all the class IDs and method IDs to their respective
 // variables. That way, later when the Java functions need to be called, the IDs will be ready.
 // It is much slower to keep looking up the class and method IDs.
+
+#if UE_BUILD_SHIPPING
+#define CHECK_JNI_RESULT( Id )
+#else
+#define CHECK_JNI_RESULT( Id ) \
+	if ( Id == 0 ) \
+	{ \
+		FPlatformMisc::LowLevelOutputDebugString(TEXT("JNI_OnLoad: Failed to find " #Id)); \
+	}
+#endif
+
 JNIEXPORT jint JNI_OnLoad(JavaVM* InJavaVM, void* InReserved)
 {
 	FPlatformMisc::LowLevelOutputDebugString(L"In the JNI_OnLoad function");
@@ -353,29 +382,55 @@ JNIEXPORT jint JNI_OnLoad(JavaVM* InJavaVM, void* InReserved)
 	JNIEnv* env = NULL;
 	InJavaVM->GetEnv((void **)&env, JNI_CURRENT_VERSION);
 
+	// if you have problems with stuff being missing esspecially in distribution builds then it could be because proguard is stripping things from java
+	// check proguard-project.txt and see if your stuff is included in the exceptions
 	GJavaVM = InJavaVM;
 
 	JDef_GameActivity::ClassID = env->FindClass("com/epicgames/ue4/GameActivity");
+	CHECK_JNI_RESULT( JDef_GameActivity::ClassID );
 
 	JDef_GameActivity::AndroidThunkJava_ShowConsoleWindow = env->GetMethodID(JDef_GameActivity::ClassID, "AndroidThunkJava_ShowConsoleWindow", "(Ljava/lang/String;)V");
+	CHECK_JNI_RESULT( JDef_GameActivity::AndroidThunkJava_ShowConsoleWindow );
 	JDef_GameActivity::AndroidThunkJava_ShowVirtualKeyboardInput = env->GetMethodID(JDef_GameActivity::ClassID, "AndroidThunkJava_ShowVirtualKeyboardInput", "(ILjava/lang/String;Ljava/lang/String;)V");
+	CHECK_JNI_RESULT( JDef_GameActivity::AndroidThunkJava_ShowVirtualKeyboardInput );
 	JDef_GameActivity::AndroidThunkJava_LaunchURL = env->GetMethodID(JDef_GameActivity::ClassID, "AndroidThunkJava_LaunchURL", "(Ljava/lang/String;)V");
+	CHECK_JNI_RESULT( JDef_GameActivity::AndroidThunkJava_LaunchURL );
 	JDef_GameActivity::AndroidThunkJava_ShowLeaderboard = env->GetMethodID(JDef_GameActivity::ClassID, "AndroidThunkJava_ShowLeaderboard", "(Ljava/lang/String;)V");
+	CHECK_JNI_RESULT( JDef_GameActivity::AndroidThunkJava_ShowLeaderboard );
 	JDef_GameActivity::AndroidThunkJava_ShowAchievements = env->GetMethodID(JDef_GameActivity::ClassID, "AndroidThunkJava_ShowAchievements", "()V");
+	CHECK_JNI_RESULT( JDef_GameActivity::AndroidThunkJava_ShowAchievements );
 	JDef_GameActivity::AndroidThunkJava_QueryAchievements = env->GetMethodID(JDef_GameActivity::ClassID, "AndroidThunkJava_QueryAchievements", "()V");
+	CHECK_JNI_RESULT( JDef_GameActivity::AndroidThunkJava_QueryAchievements );
 	JDef_GameActivity::AndroidThunkJava_ResetAchievements = env->GetMethodID(JDef_GameActivity::ClassID, "AndroidThunkJava_ResetAchievements", "()V");
 	JDef_GameActivity::AndroidThunkJava_WriteLeaderboardValue = env->GetMethodID(JDef_GameActivity::ClassID, "AndroidThunkJava_WriteLeaderboardValue", "(Ljava/lang/String;J)V");
+	CHECK_JNI_RESULT( JDef_GameActivity::AndroidThunkJava_WriteLeaderboardValue );
 	JDef_GameActivity::AndroidThunkJava_GooglePlayConnect = env->GetMethodID(JDef_GameActivity::ClassID, "AndroidThunkJava_GooglePlayConnect", "()V");
+	CHECK_JNI_RESULT( JDef_GameActivity::AndroidThunkJava_GooglePlayConnect );
 	JDef_GameActivity::AndroidThunkJava_WriteAchievement = env->GetMethodID(JDef_GameActivity::ClassID, "AndroidThunkJava_WriteAchievement", "(Ljava/lang/String;F)V");
+	CHECK_JNI_RESULT( JDef_GameActivity::AndroidThunkJava_WriteAchievement );
 	JDef_GameActivity::AndroidThunkJava_ShowAdBanner = env->GetMethodID(JDef_GameActivity::ClassID, "AndroidThunkJava_ShowAdBanner", "(Ljava/lang/String;Z)V");
+	CHECK_JNI_RESULT( JDef_GameActivity::AndroidThunkJava_ShowAdBanner );
 	JDef_GameActivity::AndroidThunkJava_HideAdBanner = env->GetMethodID(JDef_GameActivity::ClassID, "AndroidThunkJava_HideAdBanner", "()V");
+	CHECK_JNI_RESULT( JDef_GameActivity::AndroidThunkJava_HideAdBanner );
 	JDef_GameActivity::AndroidThunkJava_CloseAdBanner = env->GetMethodID(JDef_GameActivity::ClassID, "AndroidThunkJava_CloseAdBanner", "()V");
+	CHECK_JNI_RESULT( JDef_GameActivity::AndroidThunkJava_CloseAdBanner );
 	JDef_GameActivity::AndroidThunkJava_GetAssetManager = env->GetMethodID(JDef_GameActivity::ClassID, "AndroidThunkJava_GetAssetManager", "()Landroid/content/res/AssetManager;");
+	CHECK_JNI_RESULT( JDef_GameActivity::AndroidThunkJava_GetAssetManager );
+
+	JDef_GameActivity::AndroidThunkJava_Minimize = env->GetMethodID(JDef_GameActivity::ClassID, "AndroidThunkJava_Minimize", "()V");
+	CHECK_JNI_RESULT( JDef_GameActivity::AndroidThunkJava_Minimize );
+	JDef_GameActivity::AndroidThunkJava_ForceQuit = env->GetMethodID(JDef_GameActivity::ClassID, "AndroidThunkJava_ForceQuit", "()V");
+	CHECK_JNI_RESULT( JDef_GameActivity::AndroidThunkJava_ForceQuit );
+	
 
 	// Set up achievement query IDs
 	JDef_GameActivity::JavaAchievementClassID = env->FindClass("com/epicgames/ue4/GameActivity$JavaAchievement");
+	CHECK_JNI_RESULT( JDef_GameActivity::JavaAchievementClassID );
 	JDef_GameActivity::AchievementIDField = env->GetFieldID(JDef_GameActivity::JavaAchievementClassID, "ID", "Ljava/lang/String;");
+	CHECK_JNI_RESULT( JDef_GameActivity::AchievementIDField );
 	JDef_GameActivity::AchievementProgressField = env->GetFieldID(JDef_GameActivity::JavaAchievementClassID, "Progress", "D");
+	CHECK_JNI_RESULT( JDef_GameActivity::AchievementProgressField );
+	
 
 	// hook signals
 #if UE_BUILD_DEBUG
@@ -410,6 +465,9 @@ JNIEXPORT jint JNI_OnLoad(JavaVM* InJavaVM, void* InReserved)
 	extern CORE_API FAndroidLaunchURLDelegate OnAndroidLaunchURL;
 	OnAndroidLaunchURL = FAndroidLaunchURLDelegate::CreateStatic(&AndroidThunkCpp_LaunchURL);
 
+	FPlatformMisc::LowLevelOutputDebugString(L"In the JNI_OnLoad function 5");
+
+
 	return JNI_CURRENT_VERSION;
 }
 
diff --git a/Engine/Source/Runtime/Launch/Private/Android/AndroidJNI.h b/Engine/Source/Runtime/Launch/Private/Android/AndroidJNI.h
--- a/Engine/Source/Runtime/Launch/Private/Android/AndroidJNI.h
+++ b/Engine/Source/Runtime/Launch/Private/Android/AndroidJNI.h
@@ -34,6 +34,8 @@ public:
 	static jmethodID AndroidThunkJava_CloseAdBanner;
 
 	static jmethodID AndroidThunkJava_GetAssetManager;
+	static jmethodID AndroidThunkJava_Minimize;
+	static jmethodID AndroidThunkJava_ForceQuit;
 
 	// IDs related to achievement query
 	static jclass JavaAchievementClassID;
