- Fixed the landscape incorrect size on IOS8 problem [ttp 343266]

diff --git a/Engine/Source/Runtime/Launch/Private/IOS/LaunchIOS.cpp b/Engine/Source/Runtime/Launch/Private/IOS/LaunchIOS.cpp
--- a/Engine/Source/Runtime/Launch/Private/IOS/LaunchIOS.cpp
+++ b/Engine/Source/Runtime/Launch/Private/IOS/LaunchIOS.cpp
@@ -61,7 +61,15 @@ static void MainThreadInit()
 	// Size the view appropriately for any potentially dynamically attached displays,
 	// prior to creating any framebuffers
 	CGRect MainFrame = [[UIScreen mainScreen] bounds];
-	if ([IOSAppDelegate GetDelegate].OSVersion < 8.0f && !AppDelegate.bDeviceInPortraitMode)
+
+	// we need to swap if compiled with ios7, or compiled with ios8 and running on 7
+#ifndef __IPHONE_8_0
+	bool bDoLandscapeSwap = true;
+#else
+	bool bDoLandscapeSwap = AppDelegate.OSVersion < 8.0f;
+#endif
+
+	if (bDoLandscapeSwap && !AppDelegate.bDeviceInPortraitMode)
 	{
 		Swap(MainFrame.size.width, MainFrame.size.height);
 	}
