Added reference for bug comments.

diff --git a/Engine/Source/Runtime/Launch/Private/LaunchEngineLoop.cpp b/Engine/Source/Runtime/Launch/Private/LaunchEngineLoop.cpp
--- a/Engine/Source/Runtime/Launch/Private/LaunchEngineLoop.cpp
+++ b/Engine/Source/Runtime/Launch/Private/LaunchEngineLoop.cpp
@@ -1024,6 +1024,7 @@ int32 FEngineLoop::PreInit( const TCHAR* CmdLine )
 		// when running multithreaded it can't handle multiple threads using opengl (bug)
 		// tested with lg optimus 2x and motorola xoom 
 		// come back and revisit this later 
+		// https://code.google.com/p/android/issues/detail?id=32636
 		if ( FAndroidMisc::GetGPUFamily() == FString(TEXT("NVIDIA Tegra")) && FPlatformMisc::NumberOfCores() <= 2 )
 		{
 			GUseThreadedRendering = false;
@@ -1617,7 +1618,7 @@ int32 FEngineLoop::PreInit( const TCHAR* CmdLine )
 
 	//run automation smoke tests now that everything is setup to run
 	FAutomationTestFramework::GetInstance().RunSmokeTests();
-	
+
 	return 0;
 }
 
