Added "-WaitForDebugger" command-line option
- When specified, non-shipping UE4 programs will pause at startup until a debugger becomes attached
- Very useful for remote debugging code that executes at startup time (such as a plugin's StartupModule function)

diff --git a/Engine/Source/Runtime/Launch/Private/LaunchEngineLoop.cpp b/Engine/Source/Runtime/Launch/Private/LaunchEngineLoop.cpp
--- a/Engine/Source/Runtime/Launch/Private/LaunchEngineLoop.cpp
+++ b/Engine/Source/Runtime/Launch/Private/LaunchEngineLoop.cpp
@@ -2296,6 +2296,15 @@ bool FEngineLoop::AppInit( )
 	{
 		GIsBuildMachine = true;
 	}
+
+	// If "-WaitForDebugger" was specified, halt startup and wait for a debugger to attach before continuing
+	if( FParse::Param( FCommandLine::Get(), TEXT( "WaitForDebugger" ) ) )
+	{
+		while( !IsDebuggerPresent() )
+		{
+			FPlatformProcess::Sleep( 0.1f );
+		}
+	}
 #endif // !UE_BUILD_SHIPPING
 
 #if PLATFORM_WINDOWS
