CIS fix

diff --git a/Engine/Source/Runtime/Launch/Private/LaunchEngineLoop.cpp b/Engine/Source/Runtime/Launch/Private/LaunchEngineLoop.cpp
--- a/Engine/Source/Runtime/Launch/Private/LaunchEngineLoop.cpp
+++ b/Engine/Source/Runtime/Launch/Private/LaunchEngineLoop.cpp
@@ -2300,7 +2300,7 @@ bool FEngineLoop::AppInit( )
 	// If "-WaitForDebugger" was specified, halt startup and wait for a debugger to attach before continuing
 	if( FParse::Param( FCommandLine::Get(), TEXT( "WaitForDebugger" ) ) )
 	{
-		while( !IsDebuggerPresent() )
+		while( !FPlatformMisc::IsDebuggerPresent() )
 		{
 			FPlatformProcess::Sleep( 0.1f );
 		}
