Restore quotes around arguments with space.

- Minor change from PR #306 by amigo and #ue4linux folks.

diff --git a/Engine/Source/Runtime/Launch/Private/Linux/LaunchLinux.cpp b/Engine/Source/Runtime/Launch/Private/Linux/LaunchLinux.cpp
--- a/Engine/Source/Runtime/Launch/Private/Linux/LaunchLinux.cpp
+++ b/Engine/Source/Runtime/Launch/Private/Linux/LaunchLinux.cpp
@@ -132,7 +132,7 @@ int main(int argc, char *argv[])
 {
 	FPlatformMisc::SetGracefulTerminationHandler();
 
-#if UE_BUILD_SHIPPING 
+#if UE_BUILD_SHIPPING
 	// only printed in shipping
 	printf("%s %d %d %d %d\n", StringCast<ANSICHAR>(*GEngineVersion.ToString()).Get(), GEngineMinNetVersion, GEngineNegotiationVersion, GPackageFileUE4Version, GPackageFileLicenseeUE4Version);
 #endif // UE_BUILD_SHIPPING
@@ -142,7 +142,21 @@ int main(int argc, char *argv[])
 	for (int32 Option = 1; Option < argc; Option++)
 	{
 		GSavedCommandLine += TEXT(" ");
-		GSavedCommandLine += UTF8_TO_TCHAR(argv[Option]);	// note: technically it depends on locale
+		// we need to quote stuff that has spaces in it because something somewhere is removing quotation marks before they arrive here
+		FString Temp = UTF8_TO_TCHAR(argv[Option]);
+		if (Temp.Contains(TEXT(" ")))
+		{
+			if(Temp.StartsWith(TEXT("-")))
+			{
+				Temp = Temp.Replace(TEXT("="), TEXT("=\""));
+			}
+			else
+			{
+				Temp = TEXT("\"") + Temp;
+			}
+			Temp += TEXT("\"");
+		}
+		GSavedCommandLine += Temp; 	// note: technically it depends on locale
 	}
 
 #if !UE_BUILD_SHIPPING
