Added compiler version logging at engine startup
- Detects both Clang and Visual C++, prints version information to log

diff --git a/Engine/Source/Runtime/Launch/Private/LaunchEngineLoop.cpp b/Engine/Source/Runtime/Launch/Private/LaunchEngineLoop.cpp
--- a/Engine/Source/Runtime/Launch/Private/LaunchEngineLoop.cpp
+++ b/Engine/Source/Runtime/Launch/Private/LaunchEngineLoop.cpp
@@ -2421,6 +2421,26 @@ bool FEngineLoop::AppInit( )
 	UE_LOG(LogInit, Log, TEXT("Compiled (32-bit): %s %s"), ANSI_TO_TCHAR(__DATE__), ANSI_TO_TCHAR(__TIME__));
 #endif
 
+	// Print compiler version info
+#if defined(__clang__)
+	UE_LOG(LogInit, Log, TEXT("Compiled with Clang: %s"), ANSI_TO_TCHAR( __clang_version__ ) );
+#elif defined( _MSC_VER )
+	#ifndef __INTELLISENSE__	// Intellisense compiler doesn't support _MSC_FULL_VER
+	{
+		const FString VisualCPPVersion( FString::Printf( TEXT( "%d" ), _MSC_FULL_VER ) );
+		const FString VisualCPPRevisionNumber( FString::Printf( TEXT( "%02d" ), _MSC_BUILD ) );
+		UE_LOG(LogInit, Log, TEXT("Compiled with Visual C++: %s.%s.%s.%s"), 
+			*VisualCPPVersion.Mid( 0, 2 ), // Major version
+			*VisualCPPVersion.Mid( 2, 2 ), // Minor version
+			*VisualCPPVersion.Mid( 4 ),	// Build version
+			*VisualCPPRevisionNumber	// Revision number
+			);
+	}
+	#endif
+#else
+	UE_LOG(LogInit, Log, TEXT("Compiled with unrecognized C++ compiler") );
+#endif
+
 	UE_LOG(LogInit, Log, TEXT("Build Configuration: %s"), EBuildConfigurations::ToString(FApp::GetBuildConfiguration()));
 	UE_LOG(LogInit, Log, TEXT("Branch Name: %s"), *FApp::GetBranchName() );
 	UE_LOG(LogInit, Log, TEXT("Command line: %s"), FCommandLine::Get() );
