Fix opening preferences and other application dialogs on OS X post thread separation.

diff --git a/Engine/Source/Runtime/Launch/Private/Mac/LaunchMac.cpp b/Engine/Source/Runtime/Launch/Private/Mac/LaunchMac.cpp
--- a/Engine/Source/Runtime/Launch/Private/Mac/LaunchMac.cpp
+++ b/Engine/Source/Runtime/Launch/Private/Mac/LaunchMac.cpp
@@ -118,10 +118,12 @@ void EngineCrashHandler(const FGenericCrashContext & GenericContext)
 - (IBAction)showAboutWindow:(id)Sender
 {
 #if WITH_EDITOR
-	if (FModuleManager::Get().IsModuleLoaded(TEXT("MainFrame")))
-	{
-		FModuleManager::GetModuleChecked<IMainFrameModule>(TEXT("MainFrame")).ShowAboutWindow();
-	}
+	GameThreadCall(^{
+		if (FModuleManager::Get().IsModuleLoaded(TEXT("MainFrame")))
+		{
+			FModuleManager::GetModuleChecked<IMainFrameModule>(TEXT("MainFrame")).ShowAboutWindow();
+		}
+	}, false);
 #else
 	[NSApp orderFrontStandardAboutPanel:Sender];
 #endif
@@ -130,10 +132,12 @@ void EngineCrashHandler(const FGenericCrashContext & GenericContext)
 #if WITH_EDITOR
 - (IBAction)showPreferencesWindow:(id)Sender
 {
-	if (FModuleManager::Get().IsModuleLoaded(TEXT("Settings")))
-	{
-		FModuleManager::LoadModuleChecked<ISettingsModule>("Settings").ShowViewer(FName("Editor"), FName("General"), FName("Appearance"));
-	}
+	GameThreadCall(^{
+		if (FModuleManager::Get().IsModuleLoaded(TEXT("Settings")))
+		{
+			FModuleManager::LoadModuleChecked<ISettingsModule>("Settings").ShowViewer(FName("Editor"), FName("General"), FName("Appearance"));
+		}
+	}, false);
 }
 #endif
 
