Removing unnecessary modules from servers/programs.

- Do not try to load ALAudio and SteamController on Linux servers (TTP #344970).
- Do not load HeadMountedDisplay on all platform servers.

diff --git a/Engine/Source/Runtime/Core/Private/Windows/WindowsPlatformMisc.cpp b/Engine/Source/Runtime/Core/Private/Windows/WindowsPlatformMisc.cpp
--- a/Engine/Source/Runtime/Core/Private/Windows/WindowsPlatformMisc.cpp
+++ b/Engine/Source/Runtime/Core/Private/Windows/WindowsPlatformMisc.cpp
@@ -1756,7 +1756,9 @@ void FWindowsPlatformMisc::LoadPreInitModules()
 void FWindowsPlatformMisc::LoadStartupModules()
 {
 	FModuleManager::Get().LoadModule(TEXT("XAudio2"));
+#if !UE_SERVER
 	FModuleManager::Get().LoadModule(TEXT("HeadMountedDisplay"));
+#endif // !UE_SERVER
 
 #if WITH_EDITOR
 	FModuleManager::Get().LoadModule(TEXT("SourceCodeAccess"));
diff --git a/Engine/Source/Runtime/Launch/Launch.Build.cs b/Engine/Source/Runtime/Launch/Launch.Build.cs
--- a/Engine/Source/Runtime/Launch/Launch.Build.cs
+++ b/Engine/Source/Runtime/Launch/Launch.Build.cs
@@ -35,13 +35,23 @@ public class Launch : ModuleRules
 				"Slate",
 				"SlateCore",
 				"Sockets",
-				"HeadMountedDisplay",
 			}
 		);
-        
-        if( !UEBuildConfiguration.bBuildDedicatedServer)
-        {
-            PrivateIncludePathModuleNames.AddRange(
+
+		if (Target.Type != TargetRules.TargetType.Server)
+		{
+			PrivateDependencyModuleNames.AddRange(
+				new string[] {
+				"HeadMountedDisplay",
+				}
+			);
+
+			if (Target.Platform == UnrealTargetPlatform.Linux)
+			{
+				PrivateDependencyModuleNames.Add("ALAudio");
+			}
+
+			PrivateIncludePathModuleNames.AddRange(
                 new string[] {
 			        "SlateRHIRenderer",
 		        }
@@ -165,17 +175,12 @@ public class Launch : ModuleRules
 		if ((Target.Platform == UnrealTargetPlatform.Win32) ||
 			(Target.Platform == UnrealTargetPlatform.Win64) ||
 			(Target.Platform == UnrealTargetPlatform.Mac) ||
-			(Target.Platform == UnrealTargetPlatform.Linux && !UEBuildConfiguration.bBuildDedicatedServer))
+			(Target.Platform == UnrealTargetPlatform.Linux && Target.Type != TargetRules.TargetType.Server))
 		{
             // TODO: re-enable after implementing resource tables for OpenGL.
 			DynamicallyLoadedModuleNames.Add("OpenGLDrv");
 		}
 
-        if ((Target.Platform == UnrealTargetPlatform.Linux) && (Target.Type != TargetRules.TargetType.Server))
-		{
-			PrivateDependencyModuleNames.Add("ALAudio");
-		}
-
         if (Target.Platform == UnrealTargetPlatform.HTML5 )
         {
             PrivateDependencyModuleNames.Add("HTML5Audio");
diff --git a/Engine/Source/Runtime/Launch/Private/LaunchEngineLoop.cpp b/Engine/Source/Runtime/Launch/Private/LaunchEngineLoop.cpp
--- a/Engine/Source/Runtime/Launch/Private/LaunchEngineLoop.cpp
+++ b/Engine/Source/Runtime/Launch/Private/LaunchEngineLoop.cpp
@@ -45,10 +45,10 @@
 	#include "Commandlets/Commandlet.h"
 	#include "EngineService.h"
 	#include "ContentStreaming.h"
-	#include "HeadMountedDisplay.h"
 	#include "HighResScreenshot.h"
 
 #if !UE_SERVER
+	#include "HeadMountedDisplay.h"
 	#include "ISlateRHIRendererModule.h"
 #endif
 
@@ -2608,7 +2608,7 @@ void FEngineLoop::AppExit( )
 
 void FEngineLoop::PreInitHMDDevice()
 {
-#if WITH_ENGINE
+#if WITH_ENGINE && !UE_SERVER
 	if (!GIsEditor)
 	{
 		//@todo vr: only preinit first valid hmd?
@@ -2622,7 +2622,7 @@ void FEngineLoop::PreInitHMDDevice()
 			}
 		}
 	}
-#endif // #if WITH_ENGINE
+#endif // #if WITH_ENGINE && !UE_SERVER
 }
 
 #undef LOCTEXT_NAMESPACE
