Flush the log when suspending the game, so we don't end up with truncated logs on disk.

diff --git a/Engine/Source/Runtime/Launch/Private/LaunchEngineLoop.cpp b/Engine/Source/Runtime/Launch/Private/LaunchEngineLoop.cpp
--- a/Engine/Source/Runtime/Launch/Private/LaunchEngineLoop.cpp
+++ b/Engine/Source/Runtime/Launch/Private/LaunchEngineLoop.cpp
@@ -2254,6 +2254,10 @@ void FEngineLoop::OnSuspending(_In_ Platform::Object^ Sender, _In_ Windows::Appl
 	// Make the call down to the RHI to Suspend the GPU state
 	RHISuspendRendering();
 
+	// Flush the log so it's all written to disk
+	GLog->FlushThreadedLogs();
+	GLog->Flush();
+
 	// Tell the callback that we are done
 	SuspendingEvent->Complete();
 }
