Move GUseThreadedRendering Android override out of LaunchEngineLoop into new FPlatformMisc AllowRenderThread
#Android

diff --git a/Engine/Source/Runtime/Core/Private/Android/AndroidMisc.cpp b/Engine/Source/Runtime/Core/Private/Android/AndroidMisc.cpp
--- a/Engine/Source/Runtime/Core/Private/Android/AndroidMisc.cpp
+++ b/Engine/Source/Runtime/Core/Private/Android/AndroidMisc.cpp
@@ -135,9 +135,36 @@ EAppReturnType::Type FAndroidMisc::MessageBoxExt( EAppMsgType::Type MsgType, con
 
 }
 
+bool FAndroidMisc::AllowRenderThread()
+{
+	// there is a crash with the nvidia tegra dual core processors namely the optimus 2x and xoom 
+	// when running multithreaded it can't handle multiple threads using opengl (bug)
+	// tested with lg optimus 2x and motorola xoom 
+	// come back and revisit this later 
+	// https://code.google.com/p/android/issues/detail?id=32636
+	if (FAndroidMisc::GetGPUFamily() == FString(TEXT("NVIDIA Tegra")) && FPlatformMisc::NumberOfCores() <= 2)
+	{
+		return false;
+	}
+
+	// there is an issue with presenting the buffer on kindle fire (1st gen) with multiple threads using opengl
+	if (FAndroidMisc::GetDeviceModel() == FString(TEXT("Kindle Fire")))
+	{
+		return false;
+	}
+
+	// there is an issue with swapbuffer ordering on startup on samsung s3 mini with multiple threads using opengl
+	if (FAndroidMisc::GetDeviceModel() == FString(TEXT("GT-I8190L")))
+	{
+		return false;
+	}
+
+	return true;
+}
+
 int32 FAndroidMisc::NumberOfCores()
 {
-	int32 NumberOfCores = android_getCpuCount(); 
+	int32 NumberOfCores = android_getCpuCount();
 	return NumberOfCores;
 }
 
diff --git a/Engine/Source/Runtime/Core/Public/Android/AndroidMisc.h b/Engine/Source/Runtime/Core/Public/Android/AndroidMisc.h
--- a/Engine/Source/Runtime/Core/Public/Android/AndroidMisc.h
+++ b/Engine/Source/Runtime/Core/Public/Android/AndroidMisc.h
@@ -66,6 +66,7 @@ struct CORE_API FAndroidMisc : public FGenericPlatformMisc
 	static void ClipboardCopy(const TCHAR* Str);
 	static void ClipboardPaste(class FString& Dest);
 	static EAppReturnType::Type MessageBoxExt( EAppMsgType::Type MsgType, const TCHAR* Text, const TCHAR* Caption );
+	static bool AllowRenderThread();
 	static int32 NumberOfCores();
 	static void LoadPreInitModules();
 	static void BeforeRenderThreadStarts();
diff --git a/Engine/Source/Runtime/Core/Public/GenericPlatform/GenericPlatformMisc.h b/Engine/Source/Runtime/Core/Public/GenericPlatform/GenericPlatformMisc.h
--- a/Engine/Source/Runtime/Core/Public/GenericPlatform/GenericPlatformMisc.h
+++ b/Engine/Source/Runtime/Core/Public/GenericPlatform/GenericPlatformMisc.h
@@ -593,6 +593,16 @@ public:
 	*/
 	static const TCHAR* GetDefaultPathSeparator();
 
+	/**
+	 * Checks if platform wants to allow a rendering thread on current device (note: does not imply it will, only if okay given other criteria met)
+	 * @return true if allowed, false if shouldn't use a separate rendering thread
+	 */
+	static bool AllowRenderThread()
+	{
+		// allow if not overridden
+		return true;
+	}
+
 	/**
 	 * return the number of hardware CPU cores
 	 */
diff --git a/Engine/Source/Runtime/Launch/Private/LaunchEngineLoop.cpp b/Engine/Source/Runtime/Launch/Private/LaunchEngineLoop.cpp
--- a/Engine/Source/Runtime/Launch/Private/LaunchEngineLoop.cpp
+++ b/Engine/Source/Runtime/Launch/Private/LaunchEngineLoop.cpp
@@ -1028,33 +1028,9 @@ int32 FEngineLoop::PreInit( const TCHAR* CmdLine )
 	// Set all CVars which have been setup in the device profiles.
 	UDeviceProfileManager::InitializeCVarsForActiveDeviceProfile();
 
-	if (FApp::ShouldUseThreadingForPerformance())
+	if (FApp::ShouldUseThreadingForPerformance() && FPlatformMisc::AllowRenderThread())
 	{
 		GUseThreadedRendering = true;
-
-#if PLATFORM_ANDROID 
-		// there is a crash with the nvidia tegra dual core processors namely the optimus 2x and xoom 
-		// when running multithreaded it can't handle multiple threads using opengl (bug)
-		// tested with lg optimus 2x and motorola xoom 
-		// come back and revisit this later 
-		// https://code.google.com/p/android/issues/detail?id=32636
-		if ( FAndroidMisc::GetGPUFamily() == FString(TEXT("NVIDIA Tegra")) && FPlatformMisc::NumberOfCores() <= 2 )
-		{
-			GUseThreadedRendering = false;
-		}
-
-		// there is an issue with presenting the buffer on kindle fire (1st gen) with multiple threads using opengl
-		if (FAndroidMisc::GetDeviceModel() == FString(TEXT("Kindle Fire")))
-		{
-			GUseThreadedRendering = false;
-		}
-
-		// there is an issue with swapbuffer ordering on startup on samsung s3 mini with multiple threads using opengl
-		if (FAndroidMisc::GetDeviceModel() == FString(TEXT("GT-I8190L")))
-		{
-			GUseThreadedRendering = false;
-		}
-#endif
 	}
 #endif
 
