Fix for PreInit disabling log backlog when not editor to include programs not compiling against Engine

diff --git a/Engine/Source/Runtime/Launch/Private/LaunchEngineLoop.cpp b/Engine/Source/Runtime/Launch/Private/LaunchEngineLoop.cpp
--- a/Engine/Source/Runtime/Launch/Private/LaunchEngineLoop.cpp
+++ b/Engine/Source/Runtime/Launch/Private/LaunchEngineLoop.cpp
@@ -1188,11 +1188,13 @@ int32 FEngineLoop::PreInit( const TCHAR* CmdLine )
 #endif //WITH_EDITOR
 	}
 
+#endif // WITH_ENGINE
 	// If we're not in the editor stop collecting the backlog now that we know
 	if (!GIsEditor)
 	{
 		GLog->EnableBacklog( false );
 	}
+#if WITH_ENGINE
 
 	EndInitTextLocalization();
 
