#ttp 344974 WEX: Support Android System font

Changed so that on start up we find the font dir and assign it to a global.
Double checked that 'droidsans' is a safe ttf to default to - it is.

diff --git a/Engine/Build/Android/Java/src/com/epicgames/ue4/GameActivity.java b/Engine/Build/Android/Java/src/com/epicgames/ue4/GameActivity.java
--- a/Engine/Build/Android/Java/src/com/epicgames/ue4/GameActivity.java
+++ b/Engine/Build/Android/Java/src/com/epicgames/ue4/GameActivity.java
@@ -2,6 +2,8 @@
 
 package com.epicgames.ue4;
 
+import java.io.File;
+
 import java.util.Map;
 import java.util.HashMap;
 import java.util.ArrayList;
@@ -1052,6 +1054,28 @@ public class GameActivity extends NativeActivity implements GoogleApiClient.Conn
 		// finish();
 	}
 
+	public static String AndroidThunkJava_GetFontDirectory()
+	{
+		// Parse and find the first known fonts directory on the device
+		String[] fontdirs = { "/system/fonts", "/system/font", "/data/fonts" };
+
+		String targetDir = null;
+
+		for ( String fontdir : fontdirs )
+        {
+			Log.debug(fontdir);
+            File dir = new File( fontdir );
+
+			if(dir.exists())
+			{
+				targetDir = fontdir;
+				break;
+			}
+		}
+		
+		return targetDir + "/";
+	}
+
 	public native boolean nativeIsShippingBuild();
 	public native void nativeSetGlobalActivity();
 	public native void nativeSetWindowInfo(boolean bIsPortrait, int DepthBufferPreference);
diff --git a/Engine/Source/Runtime/Core/Private/Android/AndroidFile.cpp b/Engine/Source/Runtime/Core/Private/Android/AndroidFile.cpp
--- a/Engine/Source/Runtime/Core/Private/Android/AndroidFile.cpp
+++ b/Engine/Source/Runtime/Core/Private/Android/AndroidFile.cpp
@@ -21,6 +21,8 @@ static int32 GPackagePatchVersion = 0;
 
 // External File Path base - setup during load
 FString GFilePathBase;
+// External font path base - setup during load
+FString GFontPathBase;
 
 // Is the OBB in an APK file or not
 bool GOBBinAPK;
diff --git a/Engine/Source/Runtime/Core/Private/Android/AndroidMisc.cpp b/Engine/Source/Runtime/Core/Private/Android/AndroidMisc.cpp
--- a/Engine/Source/Runtime/Core/Private/Android/AndroidMisc.cpp
+++ b/Engine/Source/Runtime/Core/Private/Android/AndroidMisc.cpp
@@ -169,6 +169,7 @@ int32 FAndroidMisc::NumberOfCores()
 }
 
 extern FString GFilePathBase;
+extern FString GFontPathBase;
 
 class FTestUtime
 {
@@ -271,7 +272,8 @@ bool FAndroidMisc::GetUseVirtualJoysticks()
 TArray<uint8> FAndroidMisc::GetSystemFontBytes()
 {
 	TArray<uint8> FontBytes;
-	FFileHelper::LoadFileToArray(FontBytes, TEXT("/system/fonts/DroidSans.ttf"));
+	static FString FullFontPath = GFontPathBase + FString(TEXT("DroidSans.ttf"));
+	FFileHelper::LoadFileToArray(FontBytes, *FullFontPath);
 	return FontBytes;
 }
 
diff --git a/Engine/Source/Runtime/Launch/Private/Android/AndroidJNI.cpp b/Engine/Source/Runtime/Launch/Private/Android/AndroidJNI.cpp
--- a/Engine/Source/Runtime/Launch/Private/Android/AndroidJNI.cpp
+++ b/Engine/Source/Runtime/Launch/Private/Android/AndroidJNI.cpp
@@ -18,6 +18,7 @@ jobject GJavaGlobalThis = NULL;
 static SVirtualKeyboardEntry *VirtualKeyboardWidget = NULL;
 
 extern FString GFilePathBase;
+extern FString GFontPathBase;
 extern bool GOBBinAPK;
 
 #if USE_JNI_HELPER
@@ -120,6 +121,7 @@ jmethodID JDef_GameActivity::AndroidThunkJava_CloseAdBanner;
 jmethodID JDef_GameActivity::AndroidThunkJava_GetAssetManager;
 jmethodID JDef_GameActivity::AndroidThunkJava_Minimize;
 jmethodID JDef_GameActivity::AndroidThunkJava_ForceQuit;
+jmethodID JDef_GameActivity::AndroidThunkJava_GetFontDirectory;
 
 jclass JDef_GameActivity::JavaAchievementClassID;
 jfieldID JDef_GameActivity::AchievementIDField;
@@ -422,6 +424,9 @@ JNIEXPORT jint JNI_OnLoad(JavaVM* InJavaVM, void* InReserved)
 	JDef_GameActivity::AndroidThunkJava_ForceQuit = env->GetMethodID(JDef_GameActivity::ClassID, "AndroidThunkJava_ForceQuit", "()V");
 	CHECK_JNI_RESULT( JDef_GameActivity::AndroidThunkJava_ForceQuit );
 	
+	JDef_GameActivity::AndroidThunkJava_GetFontDirectory = env->GetStaticMethodID(JDef_GameActivity::ClassID, "AndroidThunkJava_GetFontDirectory", "()Ljava/lang/String;");
+	CHECK_JNI_RESULT( JDef_GameActivity::AndroidThunkJava_GetFontDirectory );
+	
 
 	// Set up achievement query IDs
 	JDef_GameActivity::JavaAchievementClassID = env->FindClass("com/epicgames/ue4/GameActivity$JavaAchievement");
@@ -460,6 +465,13 @@ JNIEXPORT jint JNI_OnLoad(JavaVM* InJavaVM, void* InReserved)
 	jmethodID isOBBInAPKMethod = env->GetStaticMethodID(JDef_GameActivity::ClassID, "isOBBInAPK", "()Z");
 	GOBBinAPK = (bool)env->CallStaticBooleanMethod(JDef_GameActivity::ClassID, isOBBInAPKMethod, nullptr);
 
+	// Get the system font directory
+	jstring fontPath = (jstring)env->CallStaticObjectMethod(JDef_GameActivity::ClassID, JDef_GameActivity::AndroidThunkJava_GetFontDirectory);
+	const char * nativeFontPathString = env->GetStringUTFChars(fontPath, 0);
+	GFontPathBase = FString(nativeFontPathString);
+	env->ReleaseStringUTFChars(fontPath, nativeFontPathString);
+	FPlatformMisc::LowLevelOutputDebugStringf(TEXT("Font Path found as '%s'\n"), *GFontPathBase);
+
 	// Wire up to core delegates, so core code can call out to Java
 	DECLARE_DELEGATE_OneParam(FAndroidLaunchURLDelegate, const FString&);
 	extern CORE_API FAndroidLaunchURLDelegate OnAndroidLaunchURL;
diff --git a/Engine/Source/Runtime/Launch/Private/Android/AndroidJNI.h b/Engine/Source/Runtime/Launch/Private/Android/AndroidJNI.h
--- a/Engine/Source/Runtime/Launch/Private/Android/AndroidJNI.h
+++ b/Engine/Source/Runtime/Launch/Private/Android/AndroidJNI.h
@@ -37,6 +37,8 @@ public:
 	static jmethodID AndroidThunkJava_Minimize;
 	static jmethodID AndroidThunkJava_ForceQuit;
 
+	static jmethodID AndroidThunkJava_GetFontDirectory;
+
 	// IDs related to achievement query
 	static jclass JavaAchievementClassID;
 	static jfieldID AchievementIDField;
