Add the project name to the compile-on-startup dialog.

diff --git a/Engine/Source/Runtime/Launch/Private/LaunchEngineLoop.cpp b/Engine/Source/Runtime/Launch/Private/LaunchEngineLoop.cpp
--- a/Engine/Source/Runtime/Launch/Private/LaunchEngineLoop.cpp
+++ b/Engine/Source/Runtime/Launch/Private/LaunchEngineLoop.cpp
@@ -2369,7 +2369,7 @@ bool FEngineLoop::AppInit( )
 			}
 
 			// Ask whether to compile before continuing
-			if(FPlatformMisc::MessageBoxExt(EAppMsgType::YesNo, *ModulesList, TEXT("Question")) == EAppReturnType::No)
+			if(FPlatformMisc::MessageBoxExt(EAppMsgType::YesNo, *ModulesList, *FString::Printf(TEXT("Error loading %s"), GGameName)) == EAppReturnType::No)
 			{
 				return false;
 			}
