Added support to Audio device for Android to check if background music/audio is playing.

diff --git a/Engine/Build/Android/Java/src/com/epicgames/ue4/GameActivity.java b/Engine/Build/Android/Java/src/com/epicgames/ue4/GameActivity.java
--- a/Engine/Build/Android/Java/src/com/epicgames/ue4/GameActivity.java
+++ b/Engine/Build/Android/Java/src/com/epicgames/ue4/GameActivity.java
@@ -16,6 +16,7 @@ import android.app.AlertDialog;
 import android.app.Dialog;
 import android.widget.EditText;
 import android.text.InputType;
+import android.content.Context;
 import android.content.DialogInterface;
 import android.content.Intent;
 import android.content.res.AssetManager;
@@ -39,6 +40,8 @@ import android.view.WindowManager;
 import android.widget.LinearLayout;
 import android.widget.PopupWindow;
 
+import android.media.AudioManager;
+
 import com.google.android.gms.auth.GoogleAuthUtil;
 import com.google.android.gms.common.api.GoogleApiClient;
 import com.google.android.gms.common.api.PendingResult;
@@ -1076,6 +1079,12 @@ public class GameActivity extends NativeActivity implements GoogleApiClient.Conn
 		return targetDir + "/";
 	}
 
+	public boolean AndroidThunkJava_IsMusicActive()
+	{
+		AudioManager audioManager = (AudioManager)this.getSystemService(Context.AUDIO_SERVICE);
+		return audioManager.isMusicActive();
+	}
+
 	public native boolean nativeIsShippingBuild();
 	public native void nativeSetGlobalActivity();
 	public native void nativeSetWindowInfo(boolean bIsPortrait, int DepthBufferPreference);
diff --git a/Engine/Source/Runtime/Android/AndroidAudio/Private/AndroidAudioDevice.cpp b/Engine/Source/Runtime/Android/AndroidAudio/Private/AndroidAudioDevice.cpp
--- a/Engine/Source/Runtime/Android/AndroidAudio/Private/AndroidAudioDevice.cpp
+++ b/Engine/Source/Runtime/Android/AndroidAudio/Private/AndroidAudioDevice.cpp
@@ -191,3 +191,11 @@ class ICompressedAudioInfo* FSLESAudioDevice::CreateCompressedAudioInfo(USoundWa
 
 	return NULL;
 }
+
+/** Check if any background music or sound is playing through the audio device */
+bool FSLESAudioDevice::IsExernalBackgroundSoundActive()
+{
+	extern bool AndroidThunkCpp_IsMusicActive();
+	return AndroidThunkCpp_IsMusicActive();
+
+}
diff --git a/Engine/Source/Runtime/Android/AndroidAudio/Public/AndroidAudioDevice.h b/Engine/Source/Runtime/Android/AndroidAudio/Public/AndroidAudioDevice.h
--- a/Engine/Source/Runtime/Android/AndroidAudio/Public/AndroidAudioDevice.h
+++ b/Engine/Source/Runtime/Android/AndroidAudio/Public/AndroidAudioDevice.h
@@ -279,6 +279,9 @@ public:
 	 */
 	virtual void Update( bool bGameTicking );
 
+	/** Check if any background music or sound is playing through the audio device */
+	virtual bool IsExernalBackgroundSoundActive() override;
+
 protected:
 
 	virtual FSoundSource* CreateSoundSource() override; 
diff --git a/Engine/Source/Runtime/Launch/Private/Android/AndroidJNI.cpp b/Engine/Source/Runtime/Launch/Private/Android/AndroidJNI.cpp
--- a/Engine/Source/Runtime/Launch/Private/Android/AndroidJNI.cpp
+++ b/Engine/Source/Runtime/Launch/Private/Android/AndroidJNI.cpp
@@ -122,6 +122,7 @@ jmethodID JDef_GameActivity::AndroidThunkJava_GetAssetManager;
 jmethodID JDef_GameActivity::AndroidThunkJava_Minimize;
 jmethodID JDef_GameActivity::AndroidThunkJava_ForceQuit;
 jmethodID JDef_GameActivity::AndroidThunkJava_GetFontDirectory;
+jmethodID JDef_GameActivity::AndroidThunkJava_IsMusicActive;
 
 jclass JDef_GameActivity::JavaAchievementClassID;
 jfieldID JDef_GameActivity::AchievementIDField;
@@ -356,6 +357,17 @@ void AndroidThunkCpp_ForceQuit()
 	}
 }
 
+bool AndroidThunkCpp_IsMusicActive()
+{
+	bool active = false;
+	if (JNIEnv* Env = GetJavaEnv())
+	{
+		active = (bool)Env->CallObjectMethod(GJavaGlobalThis, JDef_GameActivity::AndroidThunkJava_IsMusicActive);
+	}
+	
+	return active;
+}
+
 //The JNI_OnLoad function is triggered by loading the game library from 
 //the Java source file.
 //	static
@@ -427,6 +439,8 @@ JNIEXPORT jint JNI_OnLoad(JavaVM* InJavaVM, void* InReserved)
 	JDef_GameActivity::AndroidThunkJava_GetFontDirectory = env->GetStaticMethodID(JDef_GameActivity::ClassID, "AndroidThunkJava_GetFontDirectory", "()Ljava/lang/String;");
 	CHECK_JNI_RESULT( JDef_GameActivity::AndroidThunkJava_GetFontDirectory );
 	
+	JDef_GameActivity::AndroidThunkJava_IsMusicActive = env->GetMethodID(JDef_GameActivity::ClassID, "AndroidThunkJava_IsMusicActive", "()Z");
+	CHECK_JNI_RESULT( JDef_GameActivity::AndroidThunkJava_IsMusicActive );
 
 	// Set up achievement query IDs
 	JDef_GameActivity::JavaAchievementClassID = env->FindClass("com/epicgames/ue4/GameActivity$JavaAchievement");
diff --git a/Engine/Source/Runtime/Launch/Private/Android/AndroidJNI.h b/Engine/Source/Runtime/Launch/Private/Android/AndroidJNI.h
--- a/Engine/Source/Runtime/Launch/Private/Android/AndroidJNI.h
+++ b/Engine/Source/Runtime/Launch/Private/Android/AndroidJNI.h
@@ -39,6 +39,8 @@ public:
 
 	static jmethodID AndroidThunkJava_GetFontDirectory;
 
+	static jmethodID AndroidThunkJava_IsMusicActive;
+
 	// IDs related to achievement query
 	static jclass JavaAchievementClassID;
 	static jfieldID AchievementIDField;
