UE4: Fixed a typo in the GameplayAbilities config

diff --git a/Engine/Source/Editor/UnrealEd/Private/Editor.cpp b/Engine/Source/Editor/UnrealEd/Private/Editor.cpp
--- a/Engine/Source/Editor/UnrealEd/Private/Editor.cpp
+++ b/Engine/Source/Editor/UnrealEd/Private/Editor.cpp
@@ -587,7 +587,7 @@ void UEditorEngine::Init(IEngineLoop* InEngineLoop)
 		}
 
 		bool bGameplayAbilitiesEnabled = false;
-		GConfig->GetBool(TEXT("GameplayAbilities"), TEXT("GameplayAbilitiedEditorEnabled"), bGameplayAbilitiesEnabled, GEngineIni);
+		GConfig->GetBool(TEXT("GameplayAbilities"), TEXT("GameplayAbilitiesEditorEnabled"), bGameplayAbilitiesEnabled, GEngineIni);
 		if (bGameplayAbilitiesEnabled)
 		{
 			FModuleManager::Get().LoadModule(TEXT("GameplayAbilitiesEditor"));
diff --git a/Engine/Source/Runtime/Launch/Private/LaunchEngineLoop.cpp b/Engine/Source/Runtime/Launch/Private/LaunchEngineLoop.cpp
--- a/Engine/Source/Runtime/Launch/Private/LaunchEngineLoop.cpp
+++ b/Engine/Source/Runtime/Launch/Private/LaunchEngineLoop.cpp
@@ -1734,7 +1734,7 @@ bool FEngineLoop::LoadStartupCoreModules()
 	// HACK: load AbilitySystem editor as early as possible for statically initialized assets (non cooked BT assets needs it)
 	// cooking needs this module too
 	bool bGameplayAbilitiesEnabled = false;
-	GConfig->GetBool(TEXT("GameplayAbilities"), TEXT("GameplayAbilitiedEditorEnabled"), bGameplayAbilitiesEnabled, GEngineIni);
+	GConfig->GetBool(TEXT("GameplayAbilities"), TEXT("GameplayAbilitiesEditorEnabled"), bGameplayAbilitiesEnabled, GEngineIni);
 	if (bGameplayAbilitiesEnabled)
 	{
 		FModuleManager::Get().LoadModule(TEXT("GameplayAbilitiesEditor"));
