Made sure Mac openFile: doesn't spawn another process if the filename is not really a file path. NSApp is incorrectly treating some params in quotes as filenames.

diff --git a/Engine/Source/Runtime/Launch/Private/Mac/LaunchMac.cpp b/Engine/Source/Runtime/Launch/Private/Mac/LaunchMac.cpp
--- a/Engine/Source/Runtime/Launch/Private/Mac/LaunchMac.cpp
+++ b/Engine/Source/Runtime/Launch/Private/Mac/LaunchMac.cpp
@@ -71,10 +71,13 @@ void EngineCrashHandler(const FGenericCrashContext & GenericContext)
 {
 	if(!bHasFinishedLaunching && (GSavedCommandLine.IsEmpty() || GSavedCommandLine.Contains(FString(filename))))
 	{
-		Filename = filename;
+		if ([[NSFileManager defaultManager] fileExistsAtPath:filename])
+		{
+			Filename = filename;
+		}
 		return YES;
 	}
-	else
+	else if ([[NSFileManager defaultManager] fileExistsAtPath:filename])
 	{
 		NSString* ProjectName = [[filename stringByDeletingPathExtension] lastPathComponent];
 		
@@ -88,6 +91,10 @@ void EngineCrashHandler(const FGenericCrashContext & GenericContext)
 		
 		return (NewInstance != nil);
 	}
+	else
+	{
+		return YES;
+	}
 }
 #endif
 
