- Fixed the crash when failing to connect to a fileserver and ICU trying to localize the message, using data loaded from the server [ttp 345594]

diff --git a/Engine/Source/Runtime/Launch/Private/LaunchEngineLoop.cpp b/Engine/Source/Runtime/Launch/Private/LaunchEngineLoop.cpp
--- a/Engine/Source/Runtime/Launch/Private/LaunchEngineLoop.cpp
+++ b/Engine/Source/Runtime/Launch/Private/LaunchEngineLoop.cpp
@@ -394,7 +394,9 @@ bool LaunchCheckForFileOverride(const TCHAR* CmdLine, bool& OutFileOverrideFound
 			FPlatformMisc::LowLevelOutputDebugStringf(TEXT("Failed to connect to file server at %s. EXITING.\n"), *HostIpString);
 			uint32 Result = 2;
 #else	//PLATFORM_REQUIRES_FILESERVER
-			uint32 Result = FMessageDialog::Open( EAppMsgType::YesNoCancel, FText::Format( NSLOCTEXT("Engine", "FailedToConnectToServer", "Failed to connect to any of the following file servers:\n\n    {0}\n\nWould you like to try again? No will fallback to local disk files, Cancel will quit."), FText::FromString( HostIpString.Replace( TEXT("+"), TEXT("\n    ") ) ) ) );
+			// note that this can't be localized because it happens before we connect to a filserver - localizing would cause ICU to try to load.... from over the file server connection!
+			FString Error = FString::Printf(TEXT("Failed to connect to any of the following file servers:\n\n    %s\n\nWould you like to try again? No will fallback to local disk files, Cancel will quit."), *HostIpString.Replace( TEXT("+"), TEXT("\n    "))); 
+			uint32 Result = FMessageDialog::Open( EAppMsgType::YesNoCancel, FText::FromString( Error ) );
 #endif	//PLATFORM_REQUIRES_FILESERVER
 
 			if (Result == EAppReturnType::No)
