UE4 - turn on RHI thread by default for platforms that support it

diff --git a/Engine/Source/Runtime/Launch/Private/LaunchEngineLoop.cpp b/Engine/Source/Runtime/Launch/Private/LaunchEngineLoop.cpp
--- a/Engine/Source/Runtime/Launch/Private/LaunchEngineLoop.cpp
+++ b/Engine/Source/Runtime/Launch/Private/LaunchEngineLoop.cpp
@@ -1289,7 +1289,7 @@ int32 FEngineLoop::PreInit( const TCHAR* CmdLine )
 	if (GUseThreadedRendering)
 	{
 #if PLATFORM_SUPPORTS_RHI_THREAD
-		const bool DefaultUseRHIThread = false;
+		const bool DefaultUseRHIThread = true;
 		GUseRHIThread = DefaultUseRHIThread;
 		if (FParse::Param(FCommandLine::Get(),TEXT("rhithread")))
 		{
