CrashReport - Malloc crash for Windows

diff --git a/Engine/Source/Runtime/Launch/Private/Windows/LaunchWindows.cpp b/Engine/Source/Runtime/Launch/Private/Windows/LaunchWindows.cpp
--- a/Engine/Source/Runtime/Launch/Private/Windows/LaunchWindows.cpp
+++ b/Engine/Source/Runtime/Launch/Private/Windows/LaunchWindows.cpp
@@ -2,6 +2,7 @@
 	
 #include "LaunchPrivatePCH.h"
 #include "ExceptionHandling.h"
+#include "MallocCrash.h"
 
 #if UE_BUILD_DEBUG
 #include <crtdbg.h>
@@ -212,6 +213,7 @@ int32 WINAPI WinMain( HINSTANCE hInInstance, HINSTANCE hPrevInstance, char*, int
 			ErrorLevel = 1;
 			GError->HandleError();
 			LaunchStaticShutdownAfterError();
+			FMallocCrash::Get().PrintPoolsUsage();
 			FPlatformMisc::RequestExit( true );
 		}
 #endif
