Media: Fixed MediaAssets not being built for games

diff --git a/Engine/Source/Runtime/Launch/Launch.Build.cs b/Engine/Source/Runtime/Launch/Launch.Build.cs
--- a/Engine/Source/Runtime/Launch/Launch.Build.cs
+++ b/Engine/Source/Runtime/Launch/Launch.Build.cs
@@ -24,6 +24,7 @@ public class Launch : ModuleRules
 				"Engine",
 				"InputCore",
                 "Media",
+				"MediaAssets",
                 "MoviePlayer",
 				"Networking",
 				"PakFile",
