Removed [URL] GameName=Name from ini as we no longer need it for anything else than fixing the game name case (passed from command line)

#change Removed GGameName usage from code and replaced it with FApp::*GameName API
#change Added FApp::SetGameName(), FApp::IsGameNameEmpty()
#change Added IFileManager::GetFilenameOnDisk to get the correct case of a filename on disk
#change Renamed GGameName to GInternalGameName

diff --git a/Engine/Config/BaseEngine.ini b/Engine/Config/BaseEngine.ini
--- a/Engine/Config/BaseEngine.ini
+++ b/Engine/Config/BaseEngine.ini
@@ -9,7 +9,6 @@ GlobalDefaultGameMode="/Script/Engine.GameMode"
 LocalMapOptions=
 
 [URL]
-GameName=Engine
 Protocol=unreal
 Name=Player
 SaveExt=usa
diff --git a/Engine/Plugins/Experimental/BlueprintStats/Source/BlueprintStats/Private/BlueprintStatsModule.cpp b/Engine/Plugins/Experimental/BlueprintStats/Source/BlueprintStats/Private/BlueprintStatsModule.cpp
--- a/Engine/Plugins/Experimental/BlueprintStats/Source/BlueprintStats/Private/BlueprintStatsModule.cpp
+++ b/Engine/Plugins/Experimental/BlueprintStats/Source/BlueprintStats/Private/BlueprintStatsModule.cpp
@@ -77,7 +77,7 @@ void FBlueprintStatsModule::DumpBlueprintStats()
 	Aggregate.RemoteMacroCount.ValueSort(TGreater<int32>());
 
 	// Print out the merged record
-	UE_LOG(LogBlueprintStats, Log, TEXT("Blueprint stats for %d blueprints in %s"), Records.Num(), GGameName);
+	UE_LOG(LogBlueprintStats, Log, TEXT("Blueprint stats for %d blueprints in %s"), Records.Num(), FApp::GetGameName());
 	UE_LOG(LogBlueprintStats, Log, TEXT("%s"), *Aggregate.ToString(true));
 	UE_LOG(LogBlueprintStats, Log, TEXT("%s"), *Aggregate.ToString(false));
 	UE_LOG(LogBlueprintStats, Log, TEXT("\n"));
diff --git a/Engine/Programs/ShaderCompileWorker/Config/DefaultEngine.ini b/Engine/Programs/ShaderCompileWorker/Config/DefaultEngine.ini
--- a/Engine/Programs/ShaderCompileWorker/Config/DefaultEngine.ini
+++ b/Engine/Programs/ShaderCompileWorker/Config/DefaultEngine.ini
@@ -1,6 +1,3 @@
-[URL]
-GameName=ShaderCompileWorker
-
 [DerivedDataBackendGraph]
 Boot=(Type=Boot, Filename=../../../Engine/Programs/%GAME%/DDC/Boot.ddc, MaxCacheSize=256)
 Local=(Type=FileSystem, ReadOnly=false, Clean=false, Flush=false, Path=../../../Engine/Programs/%GAME%/DDC)
diff --git a/Engine/Programs/UnrealFrontend/Config/DefaultEngine.ini b/Engine/Programs/UnrealFrontend/Config/DefaultEngine.ini
--- a/Engine/Programs/UnrealFrontend/Config/DefaultEngine.ini
+++ b/Engine/Programs/UnrealFrontend/Config/DefaultEngine.ini
@@ -1,5 +1,4 @@
 [URL]
-GameName=UnrealFrontend
 GameNameShort=UFE
 
 [Plugins]
diff --git a/Engine/Source/Developer/LauncherServices/Private/Profiles/LauncherProfile.h b/Engine/Source/Developer/LauncherServices/Private/Profiles/LauncherProfile.h
--- a/Engine/Source/Developer/LauncherServices/Private/Profiles/LauncherProfile.h
+++ b/Engine/Source/Developer/LauncherServices/Private/Profiles/LauncherProfile.h
@@ -555,9 +555,9 @@ public:
 		{
 			ProjectPath = FPaths::GetProjectFilePath();
 		}
-		else if (FGameProjectHelper::IsGameAvailable(GGameName))
+		else if (FGameProjectHelper::IsGameAvailable(FApp::GetGameName()))
 		{
-			ProjectPath = FPaths::RootDir() / GGameName / GGameName + TEXT(".uproject");
+			ProjectPath = FPaths::RootDir() / FApp::GetGameName() / FApp::GetGameName() + TEXT(".uproject");
 		}
 		else
 		{
diff --git a/Engine/Source/Developer/ProjectLauncher/Private/Models/ProjectLauncherModel.h b/Engine/Source/Developer/ProjectLauncher/Private/Models/ProjectLauncherModel.h
--- a/Engine/Source/Developer/ProjectLauncher/Private/Models/ProjectLauncherModel.h
+++ b/Engine/Source/Developer/ProjectLauncher/Private/Models/ProjectLauncherModel.h
@@ -170,9 +170,9 @@ protected:
 		{
 			ProjectPath = FPaths::GetProjectFilePath();
 		}
-		else if (FGameProjectHelper::IsGameAvailable(GGameName))
+		else if (FGameProjectHelper::IsGameAvailable(FApp::GetGameName()))
 		{
-			ProjectPath = FPaths::RootDir() / GGameName / GGameName + TEXT(".uproject");
+			ProjectPath = FPaths::RootDir() / FApp::GetGameName() / FApp::GetGameName() + TEXT(".uproject");
 		}
 		else if (GConfig != NULL)
 		{
@@ -186,7 +186,7 @@ protected:
 	 */
 	void SaveConfig()
 	{
-		if (GConfig != NULL && !FPaths::IsProjectFilePathSet() && !FGameProjectHelper::IsGameAvailable(GGameName))
+		if (GConfig != NULL && !FPaths::IsProjectFilePathSet() && !FGameProjectHelper::IsGameAvailable(FApp::GetGameName()))
 		{
 			FString ProjectPath = ProfileManager->GetProjectPath();
 			GConfig->SetString(TEXT("FProjectLauncherModel"), TEXT("SelectedProjectPath"), *ProjectPath, GEngineIni);
diff --git a/Engine/Source/Editor/DetailCustomizations/Private/GeneralProjectSettingsDetails.cpp b/Engine/Source/Editor/DetailCustomizations/Private/GeneralProjectSettingsDetails.cpp
--- a/Engine/Source/Editor/DetailCustomizations/Private/GeneralProjectSettingsDetails.cpp
+++ b/Engine/Source/Editor/DetailCustomizations/Private/GeneralProjectSettingsDetails.cpp
@@ -20,7 +20,7 @@ void FGeneralProjectSettingsDetails::CustomizeDetails(IDetailLayoutBuilder& Deta
 
 	const FText ProjectThumbnailDesc(LOCTEXT("ProjectThumbnailLabel", "Project Thumbnail"));
 
-	const FString ProjectThumbnail_TargetImagePath = FPaths::GetPath(FPaths::GetProjectFilePath()) / FString::Printf(TEXT("%s.png"), GGameName);
+	const FString ProjectThumbnail_TargetImagePath = FPaths::GetPath(FPaths::GetProjectFilePath()) / FString::Printf(TEXT("%s.png"), FApp::GetGameName());
 	const FString ProjectThumbnail_AutomaticImagePath = FPaths::GameSavedDir() / TEXT("AutoScreenshot.png");
 
 	ImageCategory.AddCustomRow(ProjectThumbnailDesc.ToString())
diff --git a/Engine/Source/Editor/GameProjectGeneration/Private/GameProjectUtils.cpp b/Engine/Source/Editor/GameProjectGeneration/Private/GameProjectUtils.cpp
--- a/Engine/Source/Editor/GameProjectGeneration/Private/GameProjectUtils.cpp
+++ b/Engine/Source/Editor/GameProjectGeneration/Private/GameProjectUtils.cpp
@@ -1494,8 +1494,6 @@ bool GameProjectUtils::GenerateConfigFiles(const FProjectInformation& InProjectI
 		FString FileContents;
 
 		FileContents += TEXT("[URL]") LINE_TERMINATOR;
-		FileContents += FString::Printf(TEXT("GameName=%s") LINE_TERMINATOR, *NewProjectName);
-		FileContents += LINE_TERMINATOR;
 
 		FileContents += GetHardwareConfigString(InProjectInfo);
 		FileContents += LINE_TERMINATOR;
diff --git a/Engine/Source/Editor/MainFrame/Private/Frame/MainFrameActions.cpp b/Engine/Source/Editor/MainFrame/Private/Frame/MainFrameActions.cpp
--- a/Engine/Source/Editor/MainFrame/Private/Frame/MainFrameActions.cpp
+++ b/Engine/Source/Editor/MainFrame/Private/Frame/MainFrameActions.cpp
@@ -887,7 +887,7 @@ bool FMainFrameActionCallbacks::FullScreen_IsChecked()
 
 bool FMainFrameActionCallbacks::CanSwitchToProject( int32 InProjectIndex )
 {
-	if ( FApp::HasGameName() && ProjectNames[ InProjectIndex ].StartsWith( GGameName ) )
+	if (FApp::HasGameName() && ProjectNames[InProjectIndex].StartsWith(FApp::GetGameName()))
 	{
 		return false;
 	}
diff --git a/Engine/Source/Editor/PListEditor/Private/SPlistEditor.cpp b/Engine/Source/Editor/PListEditor/Private/SPlistEditor.cpp
--- a/Engine/Source/Editor/PListEditor/Private/SPlistEditor.cpp
+++ b/Engine/Source/Editor/PListEditor/Private/SPlistEditor.cpp
@@ -213,8 +213,8 @@ void SPListEditorPanel::Construct( const FArguments& InArgs )
 		]
 	];
 
-	// Default try to load the file GGameName-Info.plist file when widget is opened
-	FString DefaultFile = FString(GGameName) + TEXT("-Info.plist");
+	// Default try to load the file GameName-Info.plist file when widget is opened
+	FString DefaultFile = FString(FApp::GetGameName()) + TEXT("-Info.plist");
 	InOutLastPath += DefaultFile;
 	OpenFile(InOutLastPath);
 
@@ -825,7 +825,7 @@ bool SPListEditorPanel::OpenFile(FString FilePath)
 	if( !bLoadResult )
 	{
 		// try Info.plist
-		FilePath = FilePath.Replace(*(FString(GGameName) + TEXT("-")), TEXT(""));
+		FilePath = FilePath.Replace(*(FString(FApp::GetGameName()) + TEXT("-")), TEXT(""));
 		bLoadResult = Doc.LoadFile(FilePath);
 	}
 
@@ -914,7 +914,7 @@ FReply SPListEditorPanel::OnOpenClicked()
 		}
 
 		//FString DefaultPath = FPaths::GameDir() + TEXT("Build/IOS/");
-		//FString DefaultFile = FString(GGameName) + TEXT("-Info.plist");
+		//FString DefaultFile = FString(FApp::GetGameName()) + TEXT("-Info.plist");
 		FString FileTypes = TEXT("Property List (*.plist)|*.plist|All Files (*.*)|*.*");
 		DesktopPlatform->OpenFileDialog(
 			ParentWindowWindowHandle,
diff --git a/Engine/Source/Editor/StatsViewer/Private/SStatsViewer.cpp b/Engine/Source/Editor/StatsViewer/Private/SStatsViewer.cpp
--- a/Engine/Source/Editor/StatsViewer/Private/SStatsViewer.cpp
+++ b/Engine/Source/Editor/StatsViewer/Private/SStatsViewer.cpp
@@ -423,7 +423,7 @@ FReply SStatsViewer::OnExportClicked()
 	CSVFilename /= GWorld->GetOutermost()->GetName();
 	CSVFilename /= FString::Printf(
 		TEXT("%s-%i-%s.csv"),
-		GGameName,
+		FApp::GetGameName(),
 		GEngineVersion.GetChangelist(),
 		*FDateTime::Now().ToString() );
 
diff --git a/Engine/Source/Editor/TranslationEditor/Private/TranslationDataManager.cpp b/Engine/Source/Editor/TranslationEditor/Private/TranslationDataManager.cpp
--- a/Engine/Source/Editor/TranslationEditor/Private/TranslationDataManager.cpp
+++ b/Engine/Source/Editor/TranslationEditor/Private/TranslationDataManager.cpp
@@ -393,7 +393,7 @@ void FTranslationDataManager::GetHistoryForTranslationUnits( TArray<UTranslation
 				}
 				else
 				{
-					ProjectName = GGameName;
+					ProjectName = FApp::GetGameName();
 					SavedDir = FPaths::GameSavedDir();
 				}
 
diff --git a/Engine/Source/Editor/UnrealEd/Classes/Commandlets/ResavePackagesCommandlet.h b/Engine/Source/Editor/UnrealEd/Classes/Commandlets/ResavePackagesCommandlet.h
--- a/Engine/Source/Editor/UnrealEd/Classes/Commandlets/ResavePackagesCommandlet.h
+++ b/Engine/Source/Editor/UnrealEd/Classes/Commandlets/ResavePackagesCommandlet.h
@@ -5,7 +5,8 @@
 #include "ResavePackagesCommandlet.generated.h"
 
 UCLASS()
-class UResavePackagesCommandlet : public UCommandlet
+// Added UNREALED_API to expose this to the save packages test
+class UNREALED_API UResavePackagesCommandlet : public UCommandlet
 {
     GENERATED_UCLASS_BODY()
 
diff --git a/Engine/Source/Editor/UnrealEd/Private/Commandlets/ChunkManifestGenerator.cpp b/Engine/Source/Editor/UnrealEd/Private/Commandlets/ChunkManifestGenerator.cpp
--- a/Engine/Source/Editor/UnrealEd/Private/Commandlets/ChunkManifestGenerator.cpp
+++ b/Engine/Source/Editor/UnrealEd/Private/Commandlets/ChunkManifestGenerator.cpp
@@ -87,7 +87,7 @@ bool FChunkManifestGenerator::CleanTempPackagingDirectory(const FString& Platfor
 
 bool FChunkManifestGenerator::GenerateStreamingInstallManifest(const FString& Platform)
 {
-	FString GameNameLower = FString(GGameName).ToLower();
+	FString GameNameLower = FString(FApp::GetGameName()).ToLower();
 
 	// empty out the current paklist directory
 	FString TmpPackagingDir = GetTempPackagingDirectoryForPlatform(Platform);
diff --git a/Engine/Source/Editor/UnrealEd/Private/Commandlets/GatherTextCommandlet.cpp b/Engine/Source/Editor/UnrealEd/Private/Commandlets/GatherTextCommandlet.cpp
--- a/Engine/Source/Editor/UnrealEd/Private/Commandlets/GatherTextCommandlet.cpp
+++ b/Engine/Source/Editor/UnrealEd/Private/Commandlets/GatherTextCommandlet.cpp
@@ -212,9 +212,9 @@ FText UGatherTextCommandlet::GetChangelistDescription( const FString& InConfigPa
 	else
 	{
 		// The config file falls outside of the root directory, we will use the game name if we have it
-		if (FCString::Strlen(GGameName) != 0)
+		if (FCString::Strlen(FApp::GetGameName()) != 0)
 		{
-			ProjectName = GGameName;
+			ProjectName = FApp::GetGameName();
 		}
 	}
 
diff --git a/Engine/Source/Editor/UnrealEd/Private/Lightmass/Lightmass.cpp b/Engine/Source/Editor/UnrealEd/Private/Lightmass/Lightmass.cpp
--- a/Engine/Source/Editor/UnrealEd/Private/Lightmass/Lightmass.cpp
+++ b/Engine/Source/Editor/UnrealEd/Private/Lightmass/Lightmass.cpp
@@ -2520,7 +2520,7 @@ bool FLightmassProcessor::BeginRun()
 	FString MapNameStr = System.GetWorld()->GetMapName();
 	const TCHAR* MapName = MapNameStr.GetCharArray().GetData();
 	// Get the game name
-	const TCHAR* GameName = GGameName;
+	const TCHAR* GameName = FApp::GetGameName();
 	// Get the quality level
 	TCHAR QualityLevel[MAX_SPRINTF] = TEXT("");
 	FCString::Sprintf( QualityLevel, TEXT("%d"), ( int32 )Exporter->QualityLevel );
diff --git a/Engine/Source/Editor/UnrealEd/Private/UnrealEdMisc.cpp b/Engine/Source/Editor/UnrealEd/Private/UnrealEdMisc.cpp
--- a/Engine/Source/Editor/UnrealEd/Private/UnrealEdMisc.cpp
+++ b/Engine/Source/Editor/UnrealEd/Private/UnrealEdMisc.cpp
@@ -1378,7 +1378,7 @@ void FUnrealEdMisc::RestartEditor(bool bWarn)
 	}
 	else if(FApp::HasGameName())
 	{
-		SwitchProject(GGameName, bWarn);
+		SwitchProject(FApp::GetGameName(), bWarn);
 	}
 	else
 	{
diff --git a/Engine/Source/Runtime/Core/Private/Android/AndroidFile.cpp b/Engine/Source/Runtime/Core/Private/Android/AndroidFile.cpp
--- a/Engine/Source/Runtime/Core/Private/Android/AndroidFile.cpp
+++ b/Engine/Source/Runtime/Core/Private/Android/AndroidFile.cpp
@@ -454,6 +454,11 @@ public:
 		return AndroidEpoch + TimeSinceEpoch;
 	}
 
+	virtual FString GetFilenameOnDisk(const TCHAR* Filename) override
+	{
+		return Filename;
+	}
+
 	virtual IFileHandle* OpenRead(const TCHAR* Filename) override
 	{
 		FString NormalizedFilename = NormalizeFilename(Filename);
@@ -674,7 +679,7 @@ public:
 			Result.ReplaceInline(TEXT(".."), TEXT(""));
 			Result.ReplaceInline(FPlatformProcess::BaseDir(), TEXT(""));
 
-			static FString BasePath = GFilePathBase + FString("/") + GGameName + FString("/");
+			static FString BasePath = GFilePathBase + FString("/") + FApp::GetGameName() + FString("/");
 			Result = BasePath + Result;
 		}
 		return Result;
diff --git a/Engine/Source/Runtime/Core/Private/Apple/ApplePlatformFile.cpp b/Engine/Source/Runtime/Core/Private/Apple/ApplePlatformFile.cpp
--- a/Engine/Source/Runtime/Core/Private/Apple/ApplePlatformFile.cpp
+++ b/Engine/Source/Runtime/Core/Private/Apple/ApplePlatformFile.cpp
@@ -417,6 +417,11 @@ FDateTime FApplePlatformFile::GetAccessTimeStamp(const TCHAR* Filename)
 	return MacEpoch + TimeSinceEpoch;
 }
 
+FString FApplePlatformFile::GetFilenameOnDisk(const TCHAR* Filename)
+{
+	return Filename;
+}
+
 IFileHandle* FApplePlatformFile::OpenRead(const TCHAR* Filename)
 {
 	int32 Handle = open(TCHAR_TO_UTF8(*NormalizeFilename(Filename)), O_RDONLY);
diff --git a/Engine/Source/Runtime/Core/Private/GenericPlatform/GenericPlatformMisc.cpp b/Engine/Source/Runtime/Core/Private/GenericPlatform/GenericPlatformMisc.cpp
--- a/Engine/Source/Runtime/Core/Private/GenericPlatform/GenericPlatformMisc.cpp
+++ b/Engine/Source/Runtime/Core/Private/GenericPlatform/GenericPlatformMisc.cpp
@@ -442,7 +442,7 @@ const TCHAR* FGenericPlatformMisc::RootDir()
 
 			// keep going until we've removed Binaries
 #if IS_MONOLITHIC && !IS_PROGRAM
-			int32 pos = Path.Find(*FString::Printf(TEXT("/%s/Binaries"), GGameName));
+			int32 pos = Path.Find(*FString::Printf(TEXT("/%s/Binaries"), FApp::GetGameName()));
 #else
 			int32 pos = Path.Find(TEXT("/Engine/Binaries"), ESearchCase::IgnoreCase);
 #endif
@@ -534,7 +534,7 @@ const TCHAR* FGenericPlatformMisc::GameDir()
 {
 	static FString GameDir = TEXT("");
 
-	// track if last time we called this function the .ini was ready and had fixed the GGameName case
+	// track if last time we called this function the .ini was ready and had fixed the GameName case
 	static bool bWasIniReady = false;
 	bool bIsIniReady = GConfig && GConfig->IsReadyForUse();
 	if (bWasIniReady != bIsIniReady)
@@ -554,7 +554,7 @@ const TCHAR* FGenericPlatformMisc::GameDir()
 		if (FPlatformProperties::IsProgram())
 		{
 			// monolithic, game-agnostic executables, the ini is in Engine/Config/Platform
-			GameDir = FString::Printf(TEXT("../../../Engine/Programs/%s/"), GGameName);
+			GameDir = FString::Printf(TEXT("../../../Engine/Programs/%s/"), FApp::GetGameName());
 		}
 		else
 		{
@@ -567,7 +567,7 @@ const TCHAR* FGenericPlatformMisc::GameDir()
 				if (FPlatformProperties::IsMonolithicBuild() == false)
 				{
 					// No game project file, but has a game name, use the game folder next to the working directory
-					GameDir = FString::Printf(TEXT("../../../%s/"), GGameName);
+					GameDir = FString::Printf(TEXT("../../../%s/"), FApp::GetGameName());
 					FString GameBinariesDir = GameDir / TEXT("Binaries/");
 					if (FPlatformFileManager::Get().GetPlatformFile().DirectoryExists(*GameBinariesDir) == false)
 					{
@@ -576,7 +576,7 @@ const TCHAR* FGenericPlatformMisc::GameDir()
 						FPlatformMisc::LowLevelOutputDebugStringf(TEXT("Failed to find game directory: %s\n"), *GameDir);
 
 						// Use the uprojectdirs
-						FString GameProjectFile = FUProjectDictionary::GetDefault().GetRelativeProjectPathForGame(GGameName, FPlatformProcess::BaseDir());
+						FString GameProjectFile = FUProjectDictionary::GetDefault().GetRelativeProjectPathForGame(FApp::GetGameName(), FPlatformProcess::BaseDir());
 						if (GameProjectFile.IsEmpty() == false)
 						{
 							// We found a project folder for the game
@@ -592,7 +592,7 @@ const TCHAR* FGenericPlatformMisc::GameDir()
 				else
 				{
 #if !PLATFORM_DESKTOP
-					GameDir = FString::Printf(TEXT("../../../%s/"), GGameName);
+					GameDir = FString::Printf(TEXT("../../../%s/"), FApp::GetGameName());
 #else
 					// This assumes the game executable is in <GAME>/Binaries/<PLATFORM>
 					GameDir = TEXT("../../");
diff --git a/Engine/Source/Runtime/Core/Private/GenericPlatform/GenericPlatformOutputDevices.cpp b/Engine/Source/Runtime/Core/Private/GenericPlatform/GenericPlatformOutputDevices.cpp
--- a/Engine/Source/Runtime/Core/Private/GenericPlatform/GenericPlatformOutputDevices.cpp
+++ b/Engine/Source/Runtime/Core/Private/GenericPlatform/GenericPlatformOutputDevices.cpp
@@ -36,9 +36,9 @@ FString								FGenericPlatformOutputDevices::GetAbsoluteLogFilename()
 		if(	!FParse::Value(FCommandLine::Get(), TEXT("LOG="), Filename+FCString::Strlen(Filename), ARRAY_COUNT(Filename)-FCString::Strlen(Filename) )
 			&&	!FParse::Value(FCommandLine::Get(), TEXT("ABSLOG="), Filename, ARRAY_COUNT(Filename) ) )
 		{
-			if (FCString::Strlen(GGameName) != 0)
+			if (FCString::Strlen(FApp::GetGameName()) != 0)
 			{
-				FCString::Strcat( Filename, GGameName );
+				FCString::Strcat(Filename, FApp::GetGameName());
 			}
 			else
 			{
diff --git a/Engine/Source/Runtime/Core/Private/HAL/FileManagerGeneric.cpp b/Engine/Source/Runtime/Core/Private/HAL/FileManagerGeneric.cpp
--- a/Engine/Source/Runtime/Core/Private/HAL/FileManagerGeneric.cpp
+++ b/Engine/Source/Runtime/Core/Private/HAL/FileManagerGeneric.cpp
@@ -444,6 +444,11 @@ FDateTime FFileManagerGeneric::GetAccessTimeStamp( const TCHAR* Filename )
 	return GetLowLevel().GetAccessTimeStamp(Filename);
 }
 
+FString FFileManagerGeneric::GetFilenameOnDisk(const TCHAR* Filename)
+{
+	return GetLowLevel().GetFilenameOnDisk(Filename);
+}
+
 FString FFileManagerGeneric::DefaultConvertToRelativePath( const TCHAR* Filename )
 {
 	//default to the full absolute path of this file
diff --git a/Engine/Source/Runtime/Core/Private/HTML5/HTML5PlatformFile.cpp b/Engine/Source/Runtime/Core/Private/HTML5/HTML5PlatformFile.cpp
--- a/Engine/Source/Runtime/Core/Private/HTML5/HTML5PlatformFile.cpp
+++ b/Engine/Source/Runtime/Core/Private/HTML5/HTML5PlatformFile.cpp
@@ -285,6 +285,11 @@ public:
 		return HTML5Epoch + TimeSinceEpoch;
 	}
 
+	virtual FString GetFilenameOnDisk(const TCHAR* Filename) override
+	{
+		return Filename;
+	}
+
 	virtual IFileHandle* OpenRead(const TCHAR* Filename) override
 	{
 		FString fn = NormalizeFilename(Filename);
diff --git a/Engine/Source/Runtime/Core/Private/IOS/IOSPlatformFile.cpp b/Engine/Source/Runtime/Core/Private/IOS/IOSPlatformFile.cpp
--- a/Engine/Source/Runtime/Core/Private/IOS/IOSPlatformFile.cpp
+++ b/Engine/Source/Runtime/Core/Private/IOS/IOSPlatformFile.cpp
@@ -513,6 +513,11 @@ FDateTime FIOSPlatformFile::GetAccessTimeStamp(const TCHAR* Filename)
 	return IOSEpoch + TimeSinceEpoch;
 }
 
+FString FIOSPlatformFile::GetFilenameOnDisk(const TCHAR* Filename)
+{
+	return Filename;
+}
+
 IFileHandle* FIOSPlatformFile::OpenRead(const TCHAR* Filename)
 {
 	FString NormalizedFilename = NormalizeFilename(Filename);
diff --git a/Engine/Source/Runtime/Core/Private/IOS/IOSPlatformFile.h b/Engine/Source/Runtime/Core/Private/IOS/IOSPlatformFile.h
--- a/Engine/Source/Runtime/Core/Private/IOS/IOSPlatformFile.h
+++ b/Engine/Source/Runtime/Core/Private/IOS/IOSPlatformFile.h
@@ -28,6 +28,7 @@ public:
 	virtual FDateTime GetTimeStamp(const TCHAR* Filename) override;
 	virtual void SetTimeStamp(const TCHAR* Filename, const FDateTime DateTime) override;
 	virtual FDateTime GetAccessTimeStamp(const TCHAR* Filename) override;
+	virtual FString GetFilenameOnDisk(const TCHAR* Filename) override;
 
 	virtual IFileHandle* OpenRead(const TCHAR* Filename) override;
 	virtual IFileHandle* OpenWrite(const TCHAR* Filename, bool bAppend = false, bool bAllowRead = false) override;
diff --git a/Engine/Source/Runtime/Core/Private/Linux/LinuxPlatformFile.cpp b/Engine/Source/Runtime/Core/Private/Linux/LinuxPlatformFile.cpp
--- a/Engine/Source/Runtime/Core/Private/Linux/LinuxPlatformFile.cpp
+++ b/Engine/Source/Runtime/Core/Private/Linux/LinuxPlatformFile.cpp
@@ -226,6 +226,11 @@ FDateTime FLinuxPlatformFile::GetAccessTimeStamp(const TCHAR* Filename)
 	return UnixEpoch + TimeSinceEpoch;
 }
 
+FString FLinuxPlatformFile::GetFilenameOnDisk(const TCHAR* Filename)
+{
+	return Filename;
+}
+
 IFileHandle* FLinuxPlatformFile::OpenRead(const TCHAR* Filename)
 {
 	int32 Handle = open(TCHAR_TO_UTF8(*NormalizeFilename(Filename)), O_RDONLY);
diff --git a/Engine/Source/Runtime/Core/Private/Mac/MacPlatformMisc.cpp b/Engine/Source/Runtime/Core/Private/Mac/MacPlatformMisc.cpp
--- a/Engine/Source/Runtime/Core/Private/Mac/MacPlatformMisc.cpp
+++ b/Engine/Source/Runtime/Core/Private/Mac/MacPlatformMisc.cpp
@@ -267,7 +267,7 @@ void FMacPlatformMisc::PlatformPostInit(bool ShowSplashScreen)
 	const bool bIsBundledApp = [[[NSBundle mainBundle] bundlePath] hasSuffix:@".app"];
 	if (bIsBundledApp)
 	{
-		NSString* AppName = GIsEditor ? @"Unreal Editor" : FString(GGameName).GetNSString();
+		NSString* AppName = GIsEditor ? @"Unreal Editor" : FString(FApp::GetGameName()).GetNSString();
 
 		SEL ShowAboutSelector = [[NSApp delegate] respondsToSelector:@selector(showAboutWindow:)] ? @selector(showAboutWindow:) : @selector(orderFrontStandardAboutPanel:);
 		NSMenuItem* AboutItem = [[[NSMenuItem alloc] initWithTitle:[NSString stringWithFormat:@"About %@", AppName] action:ShowAboutSelector keyEquivalent:@""] autorelease];
diff --git a/Engine/Source/Runtime/Core/Private/Mac/MacPlatformProcess.cpp b/Engine/Source/Runtime/Core/Private/Mac/MacPlatformProcess.cpp
--- a/Engine/Source/Runtime/Core/Private/Mac/MacPlatformProcess.cpp
+++ b/Engine/Source/Runtime/Core/Private/Mac/MacPlatformProcess.cpp
@@ -612,9 +612,9 @@ const TCHAR* FMacPlatformProcess::BaseDir()
 		if ([[BasePath pathExtension] isEqual: @"app"])
 		{
 			NSString* BundledBinariesPath = NULL;
-			if (GGameName[0] != 0)
+			if (!FApp::IsGameNameEmpty())
 			{
-				BundledBinariesPath = [BasePath stringByAppendingPathComponent: [NSString stringWithFormat: @"Contents/UE4/%s/Binaries/Mac", TCHAR_TO_UTF8(GGameName)]];
+				BundledBinariesPath = [BasePath stringByAppendingPathComponent : [NSString stringWithFormat : @"Contents/UE4/%s/Binaries/Mac", TCHAR_TO_UTF8(FApp::GetGameName())]];
 			}
 			if (!BundledBinariesPath || ![FileManager fileExistsAtPath:BundledBinariesPath])
 			{
@@ -659,7 +659,7 @@ static TCHAR* UserLibrarySubDirectory()
 	static TCHAR Result[MAX_PATH] = TEXT("");
 	if (!Result[0])
 	{
-		FString SubDirectory = IsRunningGame() ? GGameName : FString(TEXT("Unreal Engine")) / GGameName;
+		FString SubDirectory = IsRunningGame() ? FString(FApp::GetGameName()) : FString(TEXT("Unreal Engine")) / FApp::GetGameName();
 		if (IsRunningDedicatedServer())
 		{
 			SubDirectory += TEXT("Server");
diff --git a/Engine/Source/Runtime/Core/Private/Misc/ConfigCacheIni.cpp b/Engine/Source/Runtime/Core/Private/Misc/ConfigCacheIni.cpp
--- a/Engine/Source/Runtime/Core/Private/Misc/ConfigCacheIni.cpp
+++ b/Engine/Source/Runtime/Core/Private/Misc/ConfigCacheIni.cpp
@@ -169,7 +169,7 @@ bool FConfigFile::Combine(const FString& Filename)
 void FConfigFile::CombineFromBuffer(const FString& Buffer)
 {
 	// Replace %GAME% with game name.
-	FString Text = Buffer.Replace( TEXT("%GAME%"), GGameName, ESearchCase::CaseSensitive );
+	FString Text = Buffer.Replace(TEXT("%GAME%"), FApp::GetGameName(), ESearchCase::CaseSensitive);
 
 	// Replace %GAMEDIR% with the game directory.
 	Text = Text.Replace( TEXT("%GAMEDIR%"), *FPaths::GameDir(), ESearchCase::CaseSensitive );
@@ -379,7 +379,7 @@ void FConfigFile::CombineFromBuffer(const FString& Buffer)
 void FConfigFile::ProcessInputFileContents(const FString& Contents)
 {
 	// Replace %GAME% with game name.
-	FString Text = Contents.Replace(TEXT("%GAME%"), GGameName, ESearchCase::CaseSensitive);
+	FString Text = Contents.Replace(TEXT("%GAME%"), FApp::GetGameName(), ESearchCase::CaseSensitive);
 
 	// Replace %GAMEDIR% with the game directory.
 	Text = Text.Replace( TEXT("%GAMEDIR%"), *FPaths::GameDir(), ESearchCase::CaseSensitive );
@@ -675,7 +675,7 @@ static bool LoadIniFileHierarchy(const TArray<FIniFilename>& HierarchyToLoad, FC
 		{
 			if (IniToLoad.bRequired)
 			{
-				//UE_LOG(LogConfig, Error, TEXT("Couldn't locate '%s' which is required to run '%s'"), *IniToLoad.Filename, GGameName );
+				//UE_LOG(LogConfig, Error, TEXT("Couldn't locate '%s' which is required to run '%s'"), *IniToLoad.Filename, FApp::GetGameName() );
 				return false;
 			}
 			else
@@ -2666,7 +2666,7 @@ void FConfigCacheIni::InitializeConfigSystem()
 
 	// load the main .ini files (unless we're running a program or a gameless UE4Editor.exe, DefaultEngine.ini is required).
 	const bool bIsGamelessExe = !FApp::HasGameName();
-	const bool bDefaultEngineIniRequired = !bIsGamelessExe && (GIsGameAgnosticExe || FCString::Strlen(GGameName) == 0);
+	const bool bDefaultEngineIniRequired = !bIsGamelessExe && (GIsGameAgnosticExe || FApp::IsGameNameEmpty());
 	bool bEngineConfigCreated = FConfigCacheIni::LoadGlobalIniFile(GEngineIni, TEXT("Engine"), NULL, NULL, bDefaultEngineIniRequired);
 
 	if ( !bIsGamelessExe )
@@ -2681,51 +2681,13 @@ void FConfigCacheIni::InitializeConfigSystem()
 			{
 				FMessageDialog::Open(EAppMsgType::Ok, Message);
 			}
-			GGameName[0] = 0; // this disables part of the crash reporter to avoid writing log files to a bogus directory
+			FApp::SetGameName(TEXT("")); // this disables part of the crash reporter to avoid writing log files to a bogus directory
 			if (!GIsBuildMachine)
 			{
 				exit(1);
 			}
 			UE_LOG(LogInit, Fatal,TEXT("%s"), *Message.ToString());
 		}
-#if PLATFORM_DESKTOP
-		// correct the case of the game name, if possible (unless we're running a program and the game name is already set)
-		const FString GameName = GConfig->GetStr(TEXT("URL"), TEXT("GameName"), GEngineIni);
-		const FString GameNameWithSuffix = GameName + TEXT("Game");
-		const bool bGameNameMatchesIniCaseSensitive = (FCString::Strcmp(*GameName, GGameName) == 0) || (FCString::Strcmp(*GameNameWithSuffix, GGameName) == 0);
-		if (!bGameNameMatchesIniCaseSensitive && (FCString::Strlen(GGameName) == 0 || GIsGameAgnosticExe || (GameName.Len() > 0 && GIsGameAgnosticExe)))
-		{
-			if (GameName == FString(GGameName)) // case insensitive compare
-			{
-				//UE_LOG(LogInit, Log, TEXT("Correcting the case of the game name to %s"), *GameName);
-				FCString::Strncpy(GGameName, *GameName, ARRAY_COUNT(GGameName));
-			}
-			else if ( GameNameWithSuffix == FString(GGameName) )
-			{
-				//UE_LOG(LogInit, Log, TEXT("Correcting the case of the game name to %sGame"), *GameName);
-				FCString::Strncpy(GGameName, *GameName, ARRAY_COUNT(GGameName));
-				FCString::Strcat(GGameName, ARRAY_COUNT(GGameName), TEXT("Game"));
-			}
-			else
-			{
-				const FText Message = FText::Format(
-					NSLOCTEXT("Core", "MismatchedGameNames", "The name of the .uproject file ('{0}') must match the GameName key in the [URL] section of Config/DefaultEngine.ini (currently '{1}').  Please either change the ini or rename the .uproject to match each other (case-insensitive match)."),
-					FText::FromString(GGameName),
-					FText::FromString(GameName));
-				if (!GIsBuildMachine)
-				{
-					UE_LOG(LogInit, Warning, TEXT("%s"), *Message.ToString());
-					FMessageDialog::Open( EAppMsgType::Ok, Message );
-				}
-				GGameName[0] = 0; // this disables part of the crash reporter to avoid writing log files to a bogus directory
-				if (!GIsBuildMachine)
-				{
-					exit(1);
-				}
-				UE_LOG(LogInit, Fatal, TEXT("%s"), *Message.ToString());
-			}
-		}
-#endif	//PLATFORM_DESKTOP
 	}
 
 	FConfigCacheIni::LoadGlobalIniFile(GGameIni, TEXT("Game"));
diff --git a/Engine/Source/Runtime/Core/Private/Misc/Core.cpp b/Engine/Source/Runtime/Core/Private/Misc/Core.cpp
--- a/Engine/Source/Runtime/Core/Private/Misc/Core.cpp
+++ b/Engine/Source/Runtime/Core/Private/Misc/Core.cpp
@@ -148,10 +148,10 @@ bool					GExitPurge						= false;
 /** Game name, used for base game directory and ini among other things										*/
 #if (!IS_MONOLITHIC && !IS_PROGRAM)
 // In modular game builds, the game name will be set when the application launches
-TCHAR					GGameName[64]					= TEXT("None");
+TCHAR					GInternalGameName[64]					= TEXT("None");
 #elif !IS_MONOLITHIC && IS_PROGRAM
 // In non-monolithic programs builds, the game name will be set by the module, but not just yet, so we need to NOT initialize it!
-TCHAR					GGameName[64];
+TCHAR					GInternalGameName[64];
 #else
 // For monolithic builds, the game name variable definition will be set by the IMPLEMENT_GAME_MODULE
 // macro for the game's main game module.
diff --git a/Engine/Source/Runtime/Core/Private/Misc/Paths.cpp b/Engine/Source/Runtime/Core/Private/Misc/Paths.cpp
--- a/Engine/Source/Runtime/Core/Private/Misc/Paths.cpp
+++ b/Engine/Source/Runtime/Core/Private/Misc/Paths.cpp
@@ -86,7 +86,7 @@ FString FPaths::GameUserDir()
 {
 	if (ShouldSaveToUserDir())
 	{
-		return FPaths::Combine(FPlatformProcess::UserSettingsDir(), GGameName) + TEXT("/");
+		return FPaths::Combine(FPlatformProcess::UserSettingsDir(), FApp::GetGameName()) + TEXT("/");
 	}
 	else
 	{
diff --git a/Engine/Source/Runtime/Core/Private/ProfilingDebugging/MallocProfiler.cpp b/Engine/Source/Runtime/Core/Private/ProfilingDebugging/MallocProfiler.cpp
--- a/Engine/Source/Runtime/Core/Private/ProfilingDebugging/MallocProfiler.cpp
+++ b/Engine/Source/Runtime/Core/Private/ProfilingDebugging/MallocProfiler.cpp
@@ -1041,7 +1041,7 @@ void FMallocProfilerBufferedFileWriter::Serialize( void* V, int64 Length )
 		if (BaseFilePath == TEXT(""))
 		{
 			const FString SysTime = FDateTime::Now().ToString();
-			BaseFilePath = FPaths::ProfilingDir() + GGameName + TEXT("-") + SysTime + TEXT("/") + GGameName;
+			BaseFilePath = FPaths::ProfilingDir() + FApp::GetGameName() + TEXT("-") + SysTime + TEXT("/") + FApp::GetGameName();
 		}
 			
 		// Create file writer to serialize data to HDD.
diff --git a/Engine/Source/Runtime/Core/Private/Windows/WindowsPlatformFile.cpp b/Engine/Source/Runtime/Core/Private/Windows/WindowsPlatformFile.cpp
--- a/Engine/Source/Runtime/Core/Private/Windows/WindowsPlatformFile.cpp
+++ b/Engine/Source/Runtime/Core/Private/Windows/WindowsPlatformFile.cpp
@@ -221,6 +221,40 @@ public:
 		return WindowsEpoch + TimeSinceEpoch;
 	}
 
+	virtual FString GetFilenameOnDisk(const TCHAR* Filename) override
+	{
+		FString Result;
+		WIN32_FIND_DATAW Data;
+		FString NormalizedFilename = NormalizeFilename(Filename);
+		while (NormalizedFilename.Len())
+		{
+			HANDLE Handle = FindFirstFileW(*NormalizedFilename, &Data);
+			if (Handle != INVALID_HANDLE_VALUE)
+			{
+				if (Result.Len())
+				{
+					Result = FString(Data.cFileName) / Result;
+				}
+				else
+				{
+					Result = Data.cFileName;
+				}				
+				FindClose(Handle);
+			}
+			int32 SeparatorIndex = INDEX_NONE;
+			if (NormalizedFilename.FindLastChar('/', SeparatorIndex))
+			{
+				NormalizedFilename = NormalizedFilename.Mid(0, SeparatorIndex);				
+			}
+			if (NormalizedFilename.Len() && (SeparatorIndex == INDEX_NONE || NormalizedFilename.EndsWith(TEXT(":"))))
+			{
+				Result = NormalizedFilename / Result;
+				NormalizedFilename.Empty();
+			}
+		}
+		return Result;
+	}
+
 	virtual IFileHandle* OpenRead(const TCHAR* Filename) override
 	{
 		uint32  Access    = GENERIC_READ;
diff --git a/Engine/Source/Runtime/Core/Private/Windows/WindowsPlatformMisc.cpp b/Engine/Source/Runtime/Core/Private/Windows/WindowsPlatformMisc.cpp
--- a/Engine/Source/Runtime/Core/Private/Windows/WindowsPlatformMisc.cpp
+++ b/Engine/Source/Runtime/Core/Private/Windows/WindowsPlatformMisc.cpp
@@ -672,7 +672,7 @@ void FWindowsPlatformMisc::SubmitErrorReport( const TCHAR* InErrorHist, EErrorRe
 				TCHAR AutoReportExe[] = TEXT("../../../Engine/Binaries/DotNET/AutoReporter.exe");
 
 				FString IniDumpPath;
-				if (GGameName[0])
+				if (!FApp::IsGameNameEmpty())
 				{
 					const TCHAR IniDumpFilename[] = TEXT("UnrealAutoReportIniDump");
 					IniDumpPath = FPaths::CreateTempFilename(*FPaths::GameLogDir(), IniDumpFilename, TEXT(".txt"));
@@ -695,10 +695,10 @@ void FWindowsPlatformMisc::SubmitErrorReport( const TCHAR* InErrorHist, EErrorRe
 				{
 					LogDirectory += CommandlineLogFile;
 				}
-				else if (FCString::Strlen(GGameName) != 0)
+				else if (!FApp::IsGameNameEmpty())
 				{
 					// If the app name is defined, use it as the log filename
-					LogDirectory += FString::Printf(TEXT("%s.Log"), GGameName);
+					LogDirectory += FString::Printf(TEXT("%s.Log"), FApp::GetGameName());
 				}
 				else
 				{
diff --git a/Engine/Source/Runtime/Core/Public/Apple/ApplePlatformFile.h b/Engine/Source/Runtime/Core/Public/Apple/ApplePlatformFile.h
--- a/Engine/Source/Runtime/Core/Public/Apple/ApplePlatformFile.h
+++ b/Engine/Source/Runtime/Core/Public/Apple/ApplePlatformFile.h
@@ -28,6 +28,7 @@ public:
 	virtual void SetTimeStamp(const TCHAR* Filename, const FDateTime DateTime) override;
 
 	virtual FDateTime GetAccessTimeStamp(const TCHAR* Filename) override;
+	virtual FString GetFilenameOnDisk(const TCHAR* Filename) override;
 
 	virtual IFileHandle* OpenRead(const TCHAR* Filename) override;
 	virtual IFileHandle* OpenWrite(const TCHAR* Filename, bool bAppend = false, bool bAllowRead = false) override;
diff --git a/Engine/Source/Runtime/Core/Public/Core.h b/Engine/Source/Runtime/Core/Public/Core.h
--- a/Engine/Source/Runtime/Core/Public/Core.h
+++ b/Engine/Source/Runtime/Core/Public/Core.h
@@ -256,7 +256,7 @@ extern CORE_API FString GGameUserSettingsIni;
 extern CORE_API float GNearClippingPlane;
 
 extern CORE_API bool GExitPurge;
-extern CORE_API TCHAR GGameName[64];
+extern CORE_API TCHAR GInternalGameName[64];
 extern CORE_API const TCHAR* GForeignEngineDir;
 
 /** Exec handler for game debugging tool, allowing commands like "editactor" */
diff --git a/Engine/Source/Runtime/Core/Public/CoreGlobals.h b/Engine/Source/Runtime/Core/Public/CoreGlobals.h
--- a/Engine/Source/Runtime/Core/Public/CoreGlobals.h
+++ b/Engine/Source/Runtime/Core/Public/CoreGlobals.h
@@ -161,7 +161,7 @@ extern CORE_API FString GGameUserSettingsIni;
 extern CORE_API float GNearClippingPlane;
 
 extern CORE_API bool GExitPurge;
-extern CORE_API TCHAR GGameName[64];
+extern CORE_API TCHAR GInternalGameName[64];
 extern CORE_API const TCHAR* GForeignEngineDir;
 
 /** Exec handler for game debugging tool, allowing commands like "editactor" */
diff --git a/Engine/Source/Runtime/Core/Public/GenericPlatform/GenericPlatformFile.h b/Engine/Source/Runtime/Core/Public/GenericPlatform/GenericPlatformFile.h
--- a/Engine/Source/Runtime/Core/Public/GenericPlatform/GenericPlatformFile.h
+++ b/Engine/Source/Runtime/Core/Public/GenericPlatform/GenericPlatformFile.h
@@ -139,6 +139,8 @@ public:
 	virtual void		SetTimeStamp(const TCHAR* Filename, FDateTime DateTime) = 0;
 	/** Return the last access time of a file. Returns FDateTime::MinValue() on failure **/
 	virtual FDateTime	GetAccessTimeStamp(const TCHAR* Filename) = 0;
+	/** For case insensitive filesystems, returns the full path of the file with the same case as in the filesystem */
+	virtual FString GetFilenameOnDisk(const TCHAR* Filename) = 0;
 
 	/** Attempt to open a file for reading. If successful will return a non-nullptr pointer. Close the file by delete'ing the handle. **/
 	virtual IFileHandle*	OpenRead(const TCHAR* Filename) = 0;
diff --git a/Engine/Source/Runtime/Core/Public/HAL/FileManager.h b/Engine/Source/Runtime/Core/Public/HAL/FileManager.h
--- a/Engine/Source/Runtime/Core/Public/HAL/FileManager.h
+++ b/Engine/Source/Runtime/Core/Public/HAL/FileManager.h
@@ -185,4 +185,13 @@ public:
 	 * @return			true if the message was sent to server and it returned success, or false if there is no server, or the command failed
 	 */
 	virtual bool SendMessageToServer(const TCHAR* Message, IPlatformFile::IFileServerMessageHandler* Handler)=0;
+
+	/**
+	* For case insensitive filesystems, returns the full path of the file with the same case as in the filesystem.
+	*
+	* @param Filename	Filename to query
+	*
+	* @return	Filename with the same case as in the filesystem.
+	*/
+	virtual FString GetFilenameOnDisk(const TCHAR* Filename) = 0;
 };
diff --git a/Engine/Source/Runtime/Core/Public/HAL/FileManagerGeneric.h b/Engine/Source/Runtime/Core/Public/HAL/FileManagerGeneric.h
--- a/Engine/Source/Runtime/Core/Public/HAL/FileManagerGeneric.h
+++ b/Engine/Source/Runtime/Core/Public/HAL/FileManagerGeneric.h
@@ -67,6 +67,7 @@ public:
 	FDateTime GetTimeStamp( const TCHAR* Filename ) override;
 	FDateTime GetAccessTimeStamp( const TCHAR* Filename ) override;
 	bool SetTimeStamp( const TCHAR* Filename, FDateTime Timestamp ) override;
+	virtual FString GetFilenameOnDisk(const TCHAR* Filename) override;
 
 	virtual uint32	Copy( const TCHAR* InDestFile, const TCHAR* InSrcFile, bool ReplaceExisting, bool EvenIfReadOnly, bool Attributes, FCopyProgress* Progress );
 	virtual bool	MakeDirectory( const TCHAR* Path, bool Tree=0 );
diff --git a/Engine/Source/Runtime/Core/Public/HAL/IPlatformFileCachedWrapper.h b/Engine/Source/Runtime/Core/Public/HAL/IPlatformFileCachedWrapper.h
--- a/Engine/Source/Runtime/Core/Public/HAL/IPlatformFileCachedWrapper.h
+++ b/Engine/Source/Runtime/Core/Public/HAL/IPlatformFileCachedWrapper.h
@@ -276,6 +276,10 @@ public:
 	{
 		return LowerLevel->GetAccessTimeStamp(Filename);
 	}
+	virtual FString	GetFilenameOnDisk(const TCHAR* Filename) override
+	{
+		return LowerLevel->GetFilenameOnDisk(Filename);
+	}
 	virtual IFileHandle*	OpenRead(const TCHAR* Filename) override
 	{
 		IFileHandle* InnerHandle=LowerLevel->OpenRead(Filename);
diff --git a/Engine/Source/Runtime/Core/Public/HAL/IPlatformFileLogWrapper.h b/Engine/Source/Runtime/Core/Public/HAL/IPlatformFileLogWrapper.h
--- a/Engine/Source/Runtime/Core/Public/HAL/IPlatformFileLogWrapper.h
+++ b/Engine/Source/Runtime/Core/Public/HAL/IPlatformFileLogWrapper.h
@@ -203,6 +203,15 @@ public:
 		FILE_LOG(LogPlatformFile, Log, TEXT("GetAccessTimeStamp return %llx [%fms]"), Result.GetTicks(), ThisTime);
 		return Result;
 	}
+	virtual FString	GetFilenameOnDisk(const TCHAR* Filename) override
+	{
+		FILE_LOG(LogPlatformFile, Log, TEXT("GetFilenameOnDisk %s"), Filename);
+		double StartTime = FPlatformTime::Seconds();
+		FString Result = LowerLevel->GetFilenameOnDisk(Filename);
+		float ThisTime = 1000.0f * float(FPlatformTime::Seconds() - StartTime);
+		FILE_LOG(LogPlatformFile, Log, TEXT("GetFilenameOnDisk return %llx [%s]"), *Result, ThisTime);
+		return Result;
+	}
 	virtual IFileHandle*	OpenRead(const TCHAR* Filename) override
 	{
 		FILE_LOG(LogPlatformFile, Log, TEXT("OpenRead %s"), Filename);
diff --git a/Engine/Source/Runtime/Core/Public/HAL/IPlatformFileOpenLogWrapper.h b/Engine/Source/Runtime/Core/Public/HAL/IPlatformFileOpenLogWrapper.h
--- a/Engine/Source/Runtime/Core/Public/HAL/IPlatformFileOpenLogWrapper.h
+++ b/Engine/Source/Runtime/Core/Public/HAL/IPlatformFileOpenLogWrapper.h
@@ -121,6 +121,10 @@ public:
 	{
 		return LowerLevel->GetAccessTimeStamp(Filename);
 	}
+	virtual FString	GetFilenameOnDisk(const TCHAR* Filename) override
+	{
+		return LowerLevel->GetFilenameOnDisk(Filename);
+	}
 	virtual IFileHandle*	OpenRead(const TCHAR* Filename) override
 	{
 		IFileHandle* Result = LowerLevel->OpenRead(Filename);
diff --git a/Engine/Source/Runtime/Core/Public/HAL/IPlatformFileProfilerWrapper.h b/Engine/Source/Runtime/Core/Public/HAL/IPlatformFileProfilerWrapper.h
--- a/Engine/Source/Runtime/Core/Public/HAL/IPlatformFileProfilerWrapper.h
+++ b/Engine/Source/Runtime/Core/Public/HAL/IPlatformFileProfilerWrapper.h
@@ -53,6 +53,7 @@ struct FProfiledFileStatsOp : public FProfiledFileStatsBase
 		SetReadOnly,
 		GetTimeStamp,
 		SetTimeStamp,
+		GetFilenameOnDisk,
 		Create,
 		Copy,
 		Iterate,
@@ -353,6 +354,15 @@ public:
 		OpStat->Duration += FPlatformTime::Seconds() * 1000 - OpStat->LastOpTime;
 		return Result;
 	}
+	virtual FString	GetFilenameOnDisk(const TCHAR* Filename) override
+	{
+		StatsType* FileStat = CreateStat(Filename);
+		FProfiledFileStatsOp* OpStat = FileStat->CreateOpStat(FProfiledFileStatsOp::GetFilenameOnDisk);
+		double StartTime = FPlatformTime::Seconds();
+		FString Result = LowerLevel->GetFilenameOnDisk(Filename);
+		OpStat->Duration += FPlatformTime::Seconds() * 1000 - OpStat->LastOpTime;
+		return Result;
+	}
 	virtual IFileHandle*	OpenRead(const TCHAR* Filename) override
 	{
 		StatsType* FileStat = CreateStat( Filename );
@@ -570,6 +580,10 @@ public:
 	{
 		return LowerLevel->GetAccessTimeStamp(Filename);
 	}
+	virtual FString	GetFilenameOnDisk(const TCHAR* Filename) override
+	{
+		return LowerLevel->GetFilenameOnDisk(Filename);
+	}
 	virtual IFileHandle*	OpenRead(const TCHAR* Filename) override
 	{
 		IFileHandle* Result = LowerLevel->OpenRead(Filename);
diff --git a/Engine/Source/Runtime/Core/Public/Linux/LinuxPlatformFile.h b/Engine/Source/Runtime/Core/Public/Linux/LinuxPlatformFile.h
--- a/Engine/Source/Runtime/Core/Public/Linux/LinuxPlatformFile.h
+++ b/Engine/Source/Runtime/Core/Public/Linux/LinuxPlatformFile.h
@@ -28,6 +28,7 @@ public:
 	virtual void SetTimeStamp(const TCHAR* Filename, const FDateTime DateTime) override;
 
 	virtual FDateTime GetAccessTimeStamp(const TCHAR* Filename) override;
+	virtual FString GetFilenameOnDisk(const TCHAR* Filename) override;
 
 	virtual IFileHandle* OpenRead(const TCHAR* Filename) override;
 	virtual IFileHandle* OpenWrite(const TCHAR* Filename, bool bAppend = false, bool bAllowRead = false) override;
diff --git a/Engine/Source/Runtime/Core/Public/Misc/App.h b/Engine/Source/Runtime/Core/Public/Misc/App.h
--- a/Engine/Source/Runtime/Core/Public/Misc/App.h
+++ b/Engine/Source/Runtime/Core/Public/Misc/App.h
@@ -52,7 +52,20 @@ public:
 	 */
 	FORCEINLINE static const TCHAR* GetGameName()
 	{
-		return GGameName;
+		return GInternalGameName;
+	}
+
+	/**
+	* Sets the name of the currently running game.
+	*
+	* @param InGameName Name of the curently running game.
+	*/
+	FORCEINLINE static void SetGameName(const TCHAR* InGameName)
+	{
+		// At the moment Strcpy is not safe as we don't check the buffer size on all platforms, so we use strncpy here.
+		FCString::Strncpy(GInternalGameName, InGameName, ARRAY_COUNT(GInternalGameName));
+		// And make sure the GameName string is null terminated.
+		GInternalGameName[ARRAY_COUNT(GInternalGameName) - 1] = 0;
 	}
 
 	/**
@@ -147,14 +160,24 @@ public:
 		return SessionOwner;
 	}
 
+	/**
+	* Reports if the game name is empty
+	*
+	* @return true if the game name is empty
+	*/
+	FORCEINLINE static bool IsGameNameEmpty()
+	{
+		return (GInternalGameName[0] == 0);
+	}
+
 	/**
 	 * Reports if the game name has been set
 	 *
 	 * @return true if the game name has been set
 	 */
-	static bool HasGameName()
+	FORCEINLINE static bool HasGameName()
 	{
-		return (GGameName[0] != 0) && (FCString::Stricmp(GGameName, TEXT("None")) != 0);
+		return (IsGameNameEmpty() == false) && (FCString::Stricmp(GInternalGameName, TEXT("None")) != 0);
 	}
 
 	/**
diff --git a/Engine/Source/Runtime/Core/Public/Misc/Paths.h b/Engine/Source/Runtime/Core/Public/Misc/Paths.h
--- a/Engine/Source/Runtime/Core/Public/Misc/Paths.h
+++ b/Engine/Source/Runtime/Core/Public/Misc/Paths.h
@@ -82,8 +82,8 @@ public:
 	static FString RootDir();
 
 	/**
-	 * Returns the base directory of the current game by looking at the global
-	 * GGameName variable. This is usually a subdirectory of the installation
+	 * Returns the base directory of the current game by looking at FApp::GetGameName().
+	 * This is usually a subdirectory of the installation
 	 * root directory and can be overridden on the command line to allow self
 	 * contained mod support.
 	 *
@@ -99,8 +99,7 @@ public:
 	static FString GameUserDir();
 
 	/**
-	 * Returns the content directory of the current game by looking at the global
-	 * GGameName variable. 
+	 * Returns the content directory of the current game by looking at FApp::GetGameName().
 	 *
 	 * @return content directory
 	 */
@@ -114,24 +113,21 @@ public:
 	static FString GameConfigDir();
 
 	/**
-	 * Returns the saved directory of the current game by looking at the global
-	 * GGameName variable. 
+	 * Returns the saved directory of the current game by looking at FApp::GetGameName().
 	 *
 	 * @return saved directory
 	 */
 	static FString GameSavedDir();
 
 	/**
-	 * Returns the intermediate directory of the current game by looking at the global
-	 * GGameName variable. 
+	 * Returns the intermediate directory of the current game by looking at FApp::GetGameName().
 	 *
 	 * @return intermediate directory
 	 */
 	static FString GameIntermediateDir();
 
 	/**
-	 * Returns the plugins directory of the current game by looking at the global
-	 * GGameName variable. 
+	 * Returns the plugins directory of the current game by looking at FApp::GetGameName().
 	 *
 	 * @return plugins directory
 	 */
diff --git a/Engine/Source/Runtime/Core/Public/Modules/ModuleManager.h b/Engine/Source/Runtime/Core/Public/Modules/ModuleManager.h
--- a/Engine/Source/Runtime/Core/Public/Modules/ModuleManager.h
+++ b/Engine/Source/Runtime/Core/Public/Modules/ModuleManager.h
@@ -707,7 +707,7 @@ class FDefaultGameModuleImpl
 	#if IS_MONOLITHIC
 		#define IMPLEMENT_APPLICATION( ModuleName, GameName ) \
 			/* For monolithic builds, we must statically define the game's name string (See Core.h) */ \
-			TCHAR GGameName[64] = TEXT( GameName ); \
+			TCHAR GInternalGameName[64] = TEXT( GameName ); \
 			IMPLEMENT_DEBUGGAME() \
 			IMPLEMENT_FOREIGN_ENGINE_DIR() \
 			IMPLEMENT_GAME_MODULE(FDefaultGameModuleImpl, ModuleName) \
@@ -722,7 +722,7 @@ class FDefaultGameModuleImpl
 			{ \
 				FAutoSet##ModuleName() \
 				{ \
-					FCString::Strncpy(GGameName, TEXT( GameName ), ARRAY_COUNT(GGameName)); \
+					FCString::Strncpy(GInternalGameName, TEXT( GameName ), ARRAY_COUNT(GInternalGameName)); \
 				} \
 			} AutoSet##ModuleName; \
 			PER_MODULE_BOILERPLATE \
@@ -740,7 +740,7 @@ class FDefaultGameModuleImpl
 
 		#define IMPLEMENT_PRIMARY_GAME_MODULE( ModuleImplClass, ModuleName, DEPRECATED_GameName ) \
 			/* For monolithic builds, we must statically define the game's name string (See Core.h) */ \
-			TCHAR GGameName[64] = TEXT( PREPROCESSOR_TO_STRING(UE_PROJECT_NAME) ); \
+			TCHAR GInternalGameName[64] = TEXT( PREPROCESSOR_TO_STRING(UE_PROJECT_NAME) ); \
 			/* Implement the GIsGameAgnosticExe variable (See Core.h). */ \
 			bool GIsGameAgnosticExe = false; \
 			IMPLEMENT_DEBUGGAME() \
@@ -757,7 +757,7 @@ class FDefaultGameModuleImpl
 
 		#define IMPLEMENT_PRIMARY_GAME_MODULE( ModuleImplClass, ModuleName, DEPRECATED_GameName ) \
 			/* For monolithic builds, we must statically define the game's name string (See Core.h) */ \
-			TCHAR GGameName[64] = TEXT( PREPROCESSOR_TO_STRING(UE_PROJECT_NAME) ); \
+			TCHAR GInternalGameName[64] = TEXT( PREPROCESSOR_TO_STRING(UE_PROJECT_NAME) ); \
 			IMPLEMENT_DEBUGGAME() \
 			PER_MODULE_BOILERPLATE \
 			IMPLEMENT_FOREIGN_ENGINE_DIR() \
diff --git a/Engine/Source/Runtime/Core/Public/ProfilingDebugging/MallocProfiler.h b/Engine/Source/Runtime/Core/Public/ProfilingDebugging/MallocProfiler.h
--- a/Engine/Source/Runtime/Core/Public/ProfilingDebugging/MallocProfiler.h
+++ b/Engine/Source/Runtime/Core/Public/ProfilingDebugging/MallocProfiler.h
@@ -138,7 +138,7 @@ class FMallocProfilerBufferedFileWriter : public FArchive
 public:
 	/** Internal file writer used to serialize to HDD. */
 	FArchive*		FileWriter;
-	/** Buffered data being serialized before GGameName has been set up. */
+	/** Buffered data being serialized before GameName has been set up. */
 	TArray<uint8>	BufferedData;
 	/** Timestamped filename with path.	*/
 	FString		FullFilepath;
diff --git a/Engine/Source/Runtime/CoreUObject/Private/Misc/PackageName.cpp b/Engine/Source/Runtime/CoreUObject/Private/Misc/PackageName.cpp
--- a/Engine/Source/Runtime/CoreUObject/Private/Misc/PackageName.cpp
+++ b/Engine/Source/Runtime/CoreUObject/Private/Misc/PackageName.cpp
@@ -198,7 +198,7 @@ private:
 		GameScriptPath         = FPaths::GameDir() / TEXT("Script/");
 		GameSavedPath          = FPaths::GameSavedDir();
 
-		FString RebasedGameDir = FString::Printf(TEXT("../../../%s/"), GGameName);
+		FString RebasedGameDir = FString::Printf(TEXT("../../../%s/"), FApp::GetGameName());
 
         GameContentPathRebased = RebasedGameDir / TEXT("Content/");
         GameScriptPathRebased  = RebasedGameDir / TEXT("Script/");
diff --git a/Engine/Source/Runtime/Engine/Private/NetworkProfiler.cpp b/Engine/Source/Runtime/Engine/Private/NetworkProfiler.cpp
--- a/Engine/Source/Runtime/Engine/Private/NetworkProfiler.cpp
+++ b/Engine/Source/Runtime/Engine/Private/NetworkProfiler.cpp
@@ -397,7 +397,7 @@ void FNetworkProfiler::TrackSessionChange( bool bShouldContinueTracking, const F
 				// Rename/ move file.
 				static int32 Salt = 0;
 				Salt++;		// Use a salt to solve the issue where this function is called so fast it produces the same time (seems to happen during seamless travel)
-				const FString FinalFileName = FPaths::ProfilingDir() + GGameName + TEXT( "-" ) + FDateTime::Now().ToString() + FString::Printf( TEXT( "[%i]" ), Salt ) + TEXT( ".nprof" );
+				const FString FinalFileName = FPaths::ProfilingDir() + FApp::GetGameName() + TEXT("-") + FDateTime::Now().ToString() + FString::Printf(TEXT("[%i]"), Salt) + TEXT(".nprof");
 				bool bWasMovedSuccessfully = IFileManager::Get().Move( *FinalFileName, *TempFileName );
 
 				// Send data to UnrealConsole to upload to DB.
diff --git a/Engine/Source/Runtime/Launch/Private/Android/LaunchAndroid.cpp b/Engine/Source/Runtime/Launch/Private/Android/LaunchAndroid.cpp
--- a/Engine/Source/Runtime/Launch/Private/Android/LaunchAndroid.cpp
+++ b/Engine/Source/Runtime/Launch/Private/Android/LaunchAndroid.cpp
@@ -199,7 +199,7 @@ static void InitCommandLine()
 	FCommandLine::Set(TEXT(""));
 
 	// read in the command line text file from the sdcard if it exists
-	FString CommandLineFilePath = GFilePathBase + FString("/") + (GGameName[0] ? GGameName : FPlatformProcess::ExecutableName()) + FString("/UE4CommandLine.txt");
+	FString CommandLineFilePath = GFilePathBase + FString("/") + (!FApp::IsGameNameEmpty() ? FApp::GetGameName() : FPlatformProcess::ExecutableName()) + FString("/UE4CommandLine.txt");
 	FILE* CommandLineFile = fopen(TCHAR_TO_UTF8(*CommandLineFilePath), "r");
 	if(CommandLineFile == NULL)
 	{
diff --git a/Engine/Source/Runtime/Launch/Private/LaunchEngineLoop.cpp b/Engine/Source/Runtime/Launch/Private/LaunchEngineLoop.cpp
--- a/Engine/Source/Runtime/Launch/Private/LaunchEngineLoop.cpp
+++ b/Engine/Source/Runtime/Launch/Private/LaunchEngineLoop.cpp
@@ -219,8 +219,8 @@ bool LaunchSetGameName(const TCHAR *InCmdLine)
 {
 	if (GIsGameAgnosticExe)
 	{
-		// Initialize GGameName to an empty string. Populate it below.
-		GGameName[0] = 0;
+		// Initialize GameName to an empty string. Populate it below.
+		FApp::SetGameName(TEXT(""));
 
 		FString ProjFilePath;
 		FString LocalGameName;
@@ -229,7 +229,7 @@ bool LaunchSetGameName(const TCHAR *InCmdLine)
 			// Only set the game name if this is NOT a program...
 			if (FPlatformProperties::IsProgram() == false)
 			{
-				FCString::Strncpy(GGameName, *LocalGameName, ARRAY_COUNT(GGameName));
+				FApp::SetGameName(*LocalGameName);
 			}
 			FPaths::SetProjectFilePath(ProjFilePath);
 		}
@@ -238,7 +238,7 @@ bool LaunchSetGameName(const TCHAR *InCmdLine)
 		{
 			// Try to use the executable name as the game name.
 			LocalGameName = FPlatformProcess::ExecutableName();
-			FCString::Strncpy(GGameName, *LocalGameName, ARRAY_COUNT(GGameName));
+			FApp::SetGameName(*LocalGameName);
 
 			// Check it's not UE4Game, otherwise assume a uproject file relative to the game project directory
 			if (LocalGameName != TEXT("UE4Game"))
@@ -255,7 +255,7 @@ bool LaunchSetGameName(const TCHAR *InCmdLine)
 			bPrinted = true;
 			if (FApp::HasGameName())
 			{
-				UE_LOG(LogInit, Display, TEXT("Running engine for game: %s"), GGameName);
+				UE_LOG(LogInit, Display, TEXT("Running engine for game: %s"), FApp::GetGameName());
 			}
 			else
 			{
@@ -278,8 +278,8 @@ bool LaunchSetGameName(const TCHAR *InCmdLine)
 		{
 			if (FPlatformProperties::RequiresCookedData())
 			{
-				// Non-agnostic exes that require cooked data cannot load projects, so make sure that the LocalGameName is the GGameName
-				if (LocalGameName != GGameName)
+				// Non-agnostic exes that require cooked data cannot load projects, so make sure that the LocalGameName is the GameName
+				if (LocalGameName != FApp::GetGameName())
 				{
 					UE_LOG(LogInit, Fatal, TEXT("Non-agnostic games cannot load projects on cooked platforms - try running UE4Game."));
 				}
@@ -287,7 +287,7 @@ bool LaunchSetGameName(const TCHAR *InCmdLine)
 			// Only set the game name if this is NOT a program...
 			if (FPlatformProperties::IsProgram() == false)
 			{
-				FCString::Strncpy(GGameName, *LocalGameName, ARRAY_COUNT(GGameName));
+				FApp::SetGameName(*LocalGameName);
 			}
 			FPaths::SetProjectFilePath(ProjFilePath);
 		}
@@ -302,6 +302,46 @@ bool LaunchSetGameName(const TCHAR *InCmdLine)
 	return true;
 }
 
+void LaunchFixGameNameCase()
+{
+#if PLATFORM_DESKTOP && !IS_PROGRAM
+	// This is to make sure this function is not misused and is only called when the game name is set
+	check(FApp::HasGameName());
+
+	// correct the case of the game name, if possible (unless we're running a program and the game name is already set)	
+	if (FPaths::IsProjectFilePathSet())
+	{
+		const FString GameName(FPaths::GetBaseFilename(IFileManager::Get().GetFilenameOnDisk(*FPaths::GetProjectFilePath())));
+
+		const bool bGameNameMatchesProjectCaseSensitive = (FCString::Strcmp(*GameName, FApp::GetGameName()) == 0);
+		if (!bGameNameMatchesProjectCaseSensitive && (FApp::IsGameNameEmpty() || GIsGameAgnosticExe || (GameName.Len() > 0 && GIsGameAgnosticExe)))
+		{
+			if (GameName == FApp::GetGameName()) // case insensitive compare
+			{
+				FApp::SetGameName(*GameName);
+			}
+			else
+			{
+				const FText Message = FText::Format(
+					NSLOCTEXT("Core", "MismatchedGameNames", "The name of the .uproject file ('{0}') must match the GameName key in the [URL] section of Config/DefaultEngine.ini (currently '{1}').  Please either change the ini or rename the .uproject to match each other (case-insensitive match)."),
+					FText::FromString(FApp::GetGameName()),
+					FText::FromString(GameName));
+				if (!GIsBuildMachine)
+				{
+					UE_LOG(LogInit, Warning, TEXT("%s"), *Message.ToString());
+					FMessageDialog::Open(EAppMsgType::Ok, Message);
+				}
+				FApp::SetGameName(TEXT("")); // this disables part of the crash reporter to avoid writing log files to a bogus directory
+				if (!GIsBuildMachine)
+				{
+					exit(1);
+				}
+				UE_LOG(LogInit, Fatal, TEXT("%s"), *Message.ToString());
+			}
+		}
+	}
+#endif	//PLATFORM_DESKTOP
+}
 
 static IPlatformFile* ConditionallyCreateFileWrapper(const TCHAR* Name, IPlatformFile* CurrentPlatformFile, const TCHAR* CommandLine, bool* OutFailedToInitialize = NULL)
 {
@@ -477,7 +517,7 @@ bool LaunchHasIncompleteGameName()
 	{
 		// Verify this is a legitimate game name
 		// Launched with a game name. See if the <GameName> folder exists. If it doesn't, it could instead be <GameName>Game
-		const FString NonSuffixedGameFolder = FPaths::RootDir() / GGameName;
+		const FString NonSuffixedGameFolder = FPaths::RootDir() / FApp::GetGameName();
 		if (FPlatformFileManager::Get().GetPlatformFile().DirectoryExists(*NonSuffixedGameFolder) == false)
 		{
 			const FString SuffixedGameFolder = NonSuffixedGameFolder + TEXT("Game");
@@ -511,7 +551,7 @@ void LaunchUpdateMostRecentProjectFile()
 			else if ( FPlatformFileManager::Get().GetPlatformFile().FileExists(*RecentProjectFileContents) )
 			{
 				// The previously loaded project file was found. Change the game name here and update the project file path
-				FCString::Strncpy(GGameName, *FPaths::GetBaseFilename(RecentProjectFileContents), ARRAY_COUNT(GGameName));
+				FApp::SetGameName(*FPaths::GetBaseFilename(RecentProjectFileContents));
 				FPaths::SetProjectFilePath(RecentProjectFileContents);
 				UE_LOG(LogInit, Display, TEXT("Loading recent project file: %s"), *RecentProjectFileContents);
 
@@ -632,7 +672,7 @@ int32 FEngineLoop::PreInit( const TCHAR* CmdLine )
 		return -1;
 	}
 
-	// Set GGameName, based on the command line
+	// Set GameName, based on the command line
 	if (LaunchSetGameName(CmdLine) == false)
 	{
 		// If it failed, do not continue
@@ -682,7 +722,7 @@ int32 FEngineLoop::PreInit( const TCHAR* CmdLine )
 				UE_LOG(LogInit, Display, TEXT("Project file not found: %s"), *ProjPath);
 				UE_LOG(LogInit, Display, TEXT("\tAttempting to find via project info helper."));
 				// Use the uprojectdirs
-				FString GameProjectFile = FUProjectDictionary::GetDefault().GetRelativeProjectPathForGame(GGameName, FPlatformProcess::BaseDir());
+				FString GameProjectFile = FUProjectDictionary::GetDefault().GetRelativeProjectPathForGame(FApp::GetGameName(), FPlatformProcess::BaseDir());
 				if (GameProjectFile.IsEmpty() == false)
 				{
 					UE_LOG(LogInit, Display, TEXT("\tFound project file %s."), *GameProjectFile);
@@ -710,7 +750,7 @@ int32 FEngineLoop::PreInit( const TCHAR* CmdLine )
 		{
 			// We did not find a non-suffixed folder and we DID find the suffixed one.
 			// The engine MUST be launched with <GameName>Game.
-			const FText GameNameText = FText::FromString( GGameName );
+			const FText GameNameText = FText::FromString(FApp::GetGameName());
 			FMessageDialog::Open(EAppMsgType::Ok, FText::Format( LOCTEXT("RequiresGamePrefix", "Error: UE4Editor does not append 'Game' to the passed in game name.\nYou must use the full name.\nYou specified '{0}', use '{0}Game'."), GameNameText ) );
 			return 1;
 		}
@@ -769,7 +809,7 @@ int32 FEngineLoop::PreInit( const TCHAR* CmdLine )
 	FString NormalizedToken = Token;
 	FPaths::NormalizeFilename(NormalizedToken);
 
-	const bool bFirstTokenIsGameName = (FApp::HasGameName() && Token == GGameName);
+	const bool bFirstTokenIsGameName = (FApp::HasGameName() && Token == FApp::GetGameName());
 	const bool bFirstTokenIsGameProjectFilePath = (FPaths::IsProjectFilePathSet() && NormalizedToken == FPaths::GetProjectFilePath());
 	const bool bFirstTokenIsGameProjectFileShortName = (FPaths::IsProjectFilePathSet() && Token == FPaths::GetCleanFilename(FPaths::GetProjectFilePath()));
 
@@ -928,6 +968,8 @@ int32 FEngineLoop::PreInit( const TCHAR* CmdLine )
 		// Tell the module manager what the game binaries folder is
 		const FString GameBinariesDirectory = FPaths::Combine( FPlatformMisc::GameDir(), TEXT( "Binaries" ), FPlatformProcess::GetBinariesSubdirectory() );
 		FModuleManager::Get().SetGameBinariesDirectory(*GameBinariesDirectory);
+
+		LaunchFixGameNameCase();
 	}
 #endif
 
@@ -2419,7 +2461,7 @@ bool FEngineLoop::AppInit( )
 			}
 
 			// Ask whether to compile before continuing
-			if(FPlatformMisc::MessageBoxExt(EAppMsgType::YesNo, *ModulesList, *FString::Printf(TEXT("Missing %s Modules"), GGameName)) == EAppReturnType::No)
+			if(FPlatformMisc::MessageBoxExt(EAppMsgType::YesNo, *ModulesList, *FString::Printf(TEXT("Missing %s Modules"), FApp::GetGameName())) == EAppReturnType::No)
 			{
 				return false;
 			}
@@ -2438,7 +2480,7 @@ bool FEngineLoop::AppInit( )
 				{
 					UE_LOG(LogInit, Warning, TEXT("Still incompatible or missing module: %s"), *StillIncompatibleFiles[Idx]);
 				}
-				FPlatformMisc::MessageBoxExt(EAppMsgType::Ok, *FString::Printf(TEXT("%s could not be compiled. Try rebuilding from source manually."), GGameName), TEXT("Error"));
+				FPlatformMisc::MessageBoxExt(EAppMsgType::Ok, *FString::Printf(TEXT("%s could not be compiled. Try rebuilding from source manually."), FApp::GetGameName()), TEXT("Error"));
 				return false;
 			}
 		}
diff --git a/Engine/Source/Runtime/Launch/Private/LaunchPrivatePCH.h b/Engine/Source/Runtime/Launch/Private/LaunchPrivatePCH.h
--- a/Engine/Source/Runtime/Launch/Private/LaunchPrivatePCH.h
+++ b/Engine/Source/Runtime/Launch/Private/LaunchPrivatePCH.h
@@ -25,6 +25,6 @@
  *****************************************************************************/
 
 /**
- * Attempts to set GGameName based on the command line and check that we have a game name
+ * Attempts to set GameName based on the command line and check that we have a game name
  */
 bool LaunchSetGameName(const TCHAR *CmdLine);
diff --git a/Engine/Source/Runtime/NetworkFile/Public/NetworkPlatformFile.h b/Engine/Source/Runtime/NetworkFile/Public/NetworkPlatformFile.h
--- a/Engine/Source/Runtime/NetworkFile/Public/NetworkPlatformFile.h
+++ b/Engine/Source/Runtime/NetworkFile/Public/NetworkPlatformFile.h
@@ -90,6 +90,10 @@ public:
 		GetFileInfo(Filename, Info);
 		return Info.AccessTimeStamp;
 	}
+	virtual FString	GetFilenameOnDisk(const TCHAR* Filename) override
+	{
+		return Filename;
+	}
 	virtual IFileHandle*	OpenRead(const TCHAR* Filename) override;
 	virtual IFileHandle*	OpenWrite(const TCHAR* Filename, bool bAppend = false, bool bAllowRead = false) override;
 	virtual bool		DirectoryExists(const TCHAR* Directory) override;
diff --git a/Engine/Source/Runtime/PakFile/Public/IPlatformFilePak.h b/Engine/Source/Runtime/PakFile/Public/IPlatformFilePak.h
--- a/Engine/Source/Runtime/PakFile/Public/IPlatformFilePak.h
+++ b/Engine/Source/Runtime/PakFile/Public/IPlatformFilePak.h
@@ -1086,6 +1086,24 @@ public:
 		return Result;
 	}
 
+	virtual FString GetFilenameOnDisk(const TCHAR* Filename) override
+	{
+		FPakFile* PakFile = NULL;
+		auto FileEntry = FindFileInPakFiles(Filename, &PakFile);
+		if (FileEntry)
+		{
+			const FString Path(FPaths::GetPath(Filename));
+			const FPakDirectory* PakDirectory = PakFile->FindDirectory(*Path);
+			const FString* RealFilename = PakDirectory->FindKey(const_cast<FPakEntry*>(FileEntry));
+			return *RealFilename;
+		}
+		else
+		{
+			// Fall back to lower level.
+			return LowerLevel->GetFilenameOnDisk(Filename);
+		}
+	}
+
 	virtual IFileHandle* OpenRead(const TCHAR* Filename) override;
 
 	virtual IFileHandle* OpenWrite(const TCHAR* Filename, bool bAppend = false, bool bAllowRead = false) override
diff --git a/Engine/Source/Runtime/SandboxFile/Public/IPlatformFileSandboxWrapper.h b/Engine/Source/Runtime/SandboxFile/Public/IPlatformFileSandboxWrapper.h
--- a/Engine/Source/Runtime/SandboxFile/Public/IPlatformFileSandboxWrapper.h
+++ b/Engine/Source/Runtime/SandboxFile/Public/IPlatformFileSandboxWrapper.h
@@ -357,6 +357,21 @@ public:
 		return Result;
 	}
 
+	virtual FString	GetFilenameOnDisk(const TCHAR* Filename) override
+	{
+		FString Result;
+		FString UserFilename(*ConvertToSandboxPath(Filename));
+		if (LowerLevel->FileExists(*UserFilename))
+		{
+			Result = LowerLevel->GetFilenameOnDisk(*UserFilename);
+		}
+		else if (OkForInnerAccess(Filename))
+		{
+			Result = LowerLevel->GetFilenameOnDisk(Filename);
+		}
+		return Result;
+	}
+
 	virtual IFileHandle*	OpenRead(const TCHAR* Filename) override
 	{
 		IFileHandle* Result = LowerLevel->OpenRead( *ConvertToSandboxPath( Filename ) );
diff --git a/Engine/Source/Runtime/UE4Game/Private/UE4Game.cpp b/Engine/Source/Runtime/UE4Game/Private/UE4Game.cpp
--- a/Engine/Source/Runtime/UE4Game/Private/UE4Game.cpp
+++ b/Engine/Source/Runtime/UE4Game/Private/UE4Game.cpp
@@ -8,7 +8,7 @@ IMPLEMENT_MODULE(FDefaultModuleImpl, UE4Game);
 #if IS_MONOLITHIC
 PER_MODULE_BOILERPLATE
 bool GIsGameAgnosticExe = true;
-TCHAR GGameName[64] = TEXT("");
+TCHAR GInternalGameName[64] = TEXT("");
 IMPLEMENT_DEBUGGAME()
 IMPLEMENT_FOREIGN_ENGINE_DIR()
 #endif
