Removed remaining flag blocking early Behavior Tree loading

Statically-initialized (ie. CDO) assets that attempted to load BTs coudl fail if the flag was still set to false in ini files.

diff --git a/Engine/Source/Runtime/Launch/Private/LaunchEngineLoop.cpp b/Engine/Source/Runtime/Launch/Private/LaunchEngineLoop.cpp
--- a/Engine/Source/Runtime/Launch/Private/LaunchEngineLoop.cpp
+++ b/Engine/Source/Runtime/Launch/Private/LaunchEngineLoop.cpp
@@ -1755,18 +1755,7 @@ bool FEngineLoop::LoadStartupCoreModules()
 #if (WITH_EDITOR && !(UE_BUILD_SHIPPING || UE_BUILD_TEST))
 	// HACK: load BT editor as early as possible for statically initialized assets (non cooked BT assets needs it)
 	// cooking needs this module too
-	bool bBehaviorTreeEditorEnabled=false;
-	GConfig->GetBool(TEXT("BehaviorTreesEd"), TEXT("BehaviorTreeEditorEnabled"), bBehaviorTreeEditorEnabled, GEngineIni);
-
-	//we can override config settings from EpicLabs
-	bool bBehaviorTreeEditorEnabledFromUserSettings=false;
-	GConfig->GetBool(TEXT("/Script/UnrealEd.EditorExperimentalSettings"), TEXT("bBehaviorTreeEditor"), bBehaviorTreeEditorEnabledFromUserSettings, GEditorUserSettingsIni);
-
-	if (bBehaviorTreeEditorEnabled || bBehaviorTreeEditorEnabledFromUserSettings)
-	{
-		//let's load BT editor even for Rocket users (ShooterGame needs it but other projects should have it disabled in config file). 
-		FModuleManager::Get().LoadModule(TEXT("BehaviorTreeEditor"));
-	}
+	FModuleManager::Get().LoadModule(TEXT("BehaviorTreeEditor"));
 
 	// -----------------------------------------------------
 
