Guarding against crash if SlateApplication is not initialized.

diff --git a/Engine/Source/Runtime/Launch/Private/LaunchEngineLoop.cpp b/Engine/Source/Runtime/Launch/Private/LaunchEngineLoop.cpp
--- a/Engine/Source/Runtime/Launch/Private/LaunchEngineLoop.cpp
+++ b/Engine/Source/Runtime/Launch/Private/LaunchEngineLoop.cpp
@@ -2181,12 +2181,12 @@ void FEngineLoop::Tick()
 
 		if (FSlateApplication::IsInitialized() && !bIdleMode)
 		{
-			FSlateApplication::Get().PollGameDeviceState();
+			FSlateApplication& SlateApp = FSlateApplication::Get();
+			SlateApp.PollGameDeviceState();
+			// Gives widgets a chance to process any accumulated input
+			SlateApp.FinishedInputThisFrame();
 		}
 
-		// Gives widgets a chance to process any accumulated input
-		FSlateApplication::Get().FinishedInputThisFrame();
-
 		GEngine->Tick( FApp::GetDeltaTime(), bIdleMode );
 
 		if (GShaderCompilingManager)
