Explicitly add the ENGINE_IS_LICENSEE_VERSION define into Version.h (previously just added when needed by UAT).

diff --git a/Engine/Source/Runtime/Launch/Resources/Version.h b/Engine/Source/Runtime/Launch/Resources/Version.h
--- a/Engine/Source/Runtime/Launch/Resources/Version.h
+++ b/Engine/Source/Runtime/Launch/Resources/Version.h
@@ -10,6 +10,8 @@
 #define ENGINE_VERSION_HIWORD 0
 #define ENGINE_VERSION_LOWORD 0
 
+#define ENGINE_IS_LICENSEE_VERSION 0
+
 #define BUILT_FROM_CHANGELIST 0
 #define BRANCH_NAME "UE4"
 
