- Initialize Android VR HMDs from the Java UI thread.
- The Native->Java->JavaUI->Native callback could probably be made more generic than this.

diff --git a/Engine/Build/Android/Java/src/com/epicgames/ue4/GameActivity.java b/Engine/Build/Android/Java/src/com/epicgames/ue4/GameActivity.java
--- a/Engine/Build/Android/Java/src/com/epicgames/ue4/GameActivity.java
+++ b/Engine/Build/Android/Java/src/com/epicgames/ue4/GameActivity.java
@@ -735,6 +735,19 @@ public class GameActivity extends NativeActivity
 		// finish();
 	}
 
+	// call back into native code from the Java UI thread, initializing any available VR HMD modules
+	public void AndroidThunkJava_InitHMDs()
+	{
+		_activity.runOnUiThread(new Runnable()
+		{
+			@Override
+			public void run()
+			{
+				nativeInitHMDs();
+			}
+		});
+	}
+
 	public static String AndroidThunkJava_GetFontDirectory()
 	{
 		// Parse and find the first known fonts directory on the device
@@ -774,6 +787,8 @@ public class GameActivity extends NativeActivity
 	
 	public native boolean nativeIsGooglePlayEnabled();
 
+	public native void nativeInitHMDs();
+
 	public native void nativeResumeMainInit();
 	
 	public native void nativeCompletedConnection(int playerID, int errorCode);
diff --git a/Engine/Source/Runtime/Launch/Private/Android/AndroidJNI.cpp b/Engine/Source/Runtime/Launch/Private/Android/AndroidJNI.cpp
--- a/Engine/Source/Runtime/Launch/Private/Android/AndroidJNI.cpp
+++ b/Engine/Source/Runtime/Launch/Private/Android/AndroidJNI.cpp
@@ -120,6 +120,7 @@ jmethodID JDef_GameActivity::AndroidThunkJava_Minimize;
 jmethodID JDef_GameActivity::AndroidThunkJava_ForceQuit;
 jmethodID JDef_GameActivity::AndroidThunkJava_GetFontDirectory;
 jmethodID JDef_GameActivity::AndroidThunkJava_IsMusicActive;
+jmethodID JDef_GameActivity::AndroidThunkJava_InitHMDs;
 
 DEFINE_LOG_CATEGORY_STATIC(LogEngine, Log, All);
 
@@ -163,6 +164,15 @@ void AndroidThunkCpp_Vibrate(int64_t Duration)
 	}
 }
 
+// Call the Java side code for initializing VR HMD modules
+void AndroidThunkCpp_InitHMDs()
+{
+	if (JNIEnv* Env = GetJavaEnv())
+	{
+		Env->CallVoidMethod(GJavaGlobalThis, JDef_GameActivity::AndroidThunkJava_InitHMDs);
+	}
+}
+
 void AndroidThunkCpp_ShowConsoleWindow()
 {
 	if (JNIEnv* Env = GetJavaEnv())
@@ -390,6 +400,9 @@ JNIEXPORT jint JNI_OnLoad(JavaVM* InJavaVM, void* InReserved)
 	JDef_GameActivity::AndroidThunkJava_IsMusicActive = env->GetMethodID(JDef_GameActivity::ClassID, "AndroidThunkJava_IsMusicActive", "()Z");
 	CHECK_JNI_RESULT( JDef_GameActivity::AndroidThunkJava_IsMusicActive );	
 
+	JDef_GameActivity::AndroidThunkJava_InitHMDs = env->GetMethodID(JDef_GameActivity::ClassID, "AndroidThunkJava_InitHMDs", "()V");
+	CHECK_JNI_RESULT(JDef_GameActivity::AndroidThunkJava_InitHMDs);
+
 	// hook signals
 #if UE_BUILD_DEBUG
 	if( GAlwaysReportCrash )
diff --git a/Engine/Source/Runtime/Launch/Private/Android/AndroidJNI.h b/Engine/Source/Runtime/Launch/Private/Android/AndroidJNI.h
--- a/Engine/Source/Runtime/Launch/Private/Android/AndroidJNI.h
+++ b/Engine/Source/Runtime/Launch/Private/Android/AndroidJNI.h
@@ -33,6 +33,8 @@ public:
 	static jmethodID AndroidThunkJava_Minimize;
 	static jmethodID AndroidThunkJava_ForceQuit;
 
+	static jmethodID AndroidThunkJava_InitHMDs;
+
 	static jmethodID AndroidThunkJava_GetFontDirectory;
 
 	static jmethodID AndroidThunkJava_IsMusicActive;
diff --git a/Engine/Source/Runtime/Launch/Private/Android/LaunchAndroid.cpp b/Engine/Source/Runtime/Launch/Private/Android/LaunchAndroid.cpp
--- a/Engine/Source/Runtime/Launch/Private/Android/LaunchAndroid.cpp
+++ b/Engine/Source/Runtime/Launch/Private/Android/LaunchAndroid.cpp
@@ -90,6 +90,7 @@ extern "C"
 	void *__dso_handle;
 }
 
+extern void AndroidThunkCpp_InitHMDs();
 extern void AndroidThunkCpp_ShowConsoleWindow();
 
 // Base path for file accesses
@@ -145,6 +146,27 @@ extern "C" void Java_com_epicgames_ue4_GameActivity_nativeResumeMainInit(JNIEnv*
 	}
 }
 
+static volatile bool GHMDsInitialized = false;
+static TArray<IHeadMountedDisplayModule*> GHMDImplementations;
+void InitHMDs()
+{
+	if (FParse::Param(FCommandLine::Get(), TEXT("nohmd")) || FParse::Param(FCommandLine::Get(), TEXT("emulatestereo")))
+	{
+		return;
+	}
+
+	// Get a list of plugins that implement this feature
+	GHMDImplementations = IModularFeatures::Get().GetModularFeatureImplementations<IHeadMountedDisplayModule>(IHeadMountedDisplayModule::GetModularFeatureName());
+
+	AndroidThunkCpp_InitHMDs();
+
+	while (!GHMDsInitialized)
+	{
+		FPlatformProcess::Sleep(0.01f);
+		FPlatformMisc::MemoryBarrier();
+	}
+}
+
 
 void UpdateGameInterruptions()
 {
@@ -315,6 +337,9 @@ int32 AndroidMain(struct android_app* state)
 	// initialize the engine
 	GEngineLoop.PreInit(0, NULL, FCommandLine::Get());
 
+	// initialize HMDs
+	InitHMDs();
+
 	UE_LOG(LogAndroid, Display, TEXT("Passed PreInit()"));
 
 	GLog->SetCurrentThreadAsMasterThread();
@@ -951,6 +976,16 @@ extern "C" jboolean Java_com_epicgames_ue4_GameActivity_nativeIsGooglePlayEnable
 	return bEnabled;
 }
 
+// This is called from the Java UI thread for initializing VR HMDs
+extern "C" void Java_com_epicgames_ue4_GameActivity_nativeInitHMDs(JNIEnv* jenv, jobject thiz)
+{
+	for (auto HMDModuleIt = GHMDImplementations.CreateIterator(); HMDModuleIt; ++HMDModuleIt)
+	{
+		(*HMDModuleIt)->PreInit();
+	}
+
+	GHMDsInitialized = true;
+}
 
 extern "C" void Java_com_epicgames_ue4_GameActivity_nativeSetAndroidVersionInformation(JNIEnv* jenv, jobject thiz, jstring androidVersion, jstring phoneMake, jstring phoneModel, jstring osLanguage )
 {
