Fix for cooking problems in VehicleGame (et al.).

Turnes out that Pre/PostWorldInitialization events are not fired every time after creation and cooking initialization is completely different than normal editor run. Added a new event called in UWorld constructors and registered Landscape event to it and added auto-startup module loading to the commandlet path.

diff --git a/Engine/Source/Runtime/Engine/Classes/Engine/World.h b/Engine/Source/Runtime/Engine/Classes/Engine/World.h
--- a/Engine/Source/Runtime/Engine/Classes/Engine/World.h
+++ b/Engine/Source/Runtime/Engine/Classes/Engine/World.h
@@ -2475,6 +2475,9 @@ public:
 	DECLARE_MULTICAST_DELEGATE_TwoParams(FWorldPostDuplicateEvent, UWorld* /*World*/, bool /*bDuplicateForPIE*/);
 	// Delegate type for level change events
 	DECLARE_MULTICAST_DELEGATE_TwoParams(FOnLevelChanged, ULevel*, UWorld*);
+
+	// Callback for world creation
+	static FWorldEvent OnPostWorldCreation;
 	
 	// Callback for world initialization (pre)
 	static FWorldInitializationEvent OnPreWorldInitialization;
diff --git a/Engine/Source/Runtime/Engine/Private/World.cpp b/Engine/Source/Runtime/Engine/Private/World.cpp
--- a/Engine/Source/Runtime/Engine/Private/World.cpp
+++ b/Engine/Source/Runtime/Engine/Private/World.cpp
@@ -69,6 +69,7 @@ UWorldProxy GWorld;
 
 TMap<FName, EWorldType::Type> UWorld::WorldTypePreLoadMap;
 
+FWorldDelegates::FWorldEvent FWorldDelegates::OnPostWorldCreation;
 FWorldDelegates::FWorldInitializationEvent FWorldDelegates::OnPreWorldInitialization;
 FWorldDelegates::FWorldInitializationEvent FWorldDelegates::OnPostWorldInitialization;
 FWorldDelegates::FWorldPostDuplicateEvent FWorldDelegates::OnPostDuplicate;
@@ -89,6 +90,8 @@ UWorld::UWorld( const FObjectInitializer& ObjectInitializer )
 #if WITH_EDITOR
 	bBroadcastSelectionChange = true; //Ed Only
 #endif // WITH_EDITOR
+
+	FWorldDelegates::OnPostWorldCreation.Broadcast(this);
 }
 
 
@@ -106,6 +109,8 @@ UWorld::UWorld( const FObjectInitializer& ObjectInitializer,const FURL& InURL )
 #if WITH_EDITOR
 	bBroadcastSelectionChange = true;
 #endif // WITH_EDITOR
+
+	FWorldDelegates::OnPostWorldCreation.Broadcast(this);
 }
 
 
diff --git a/Engine/Source/Runtime/Landscape/Private/LandscapeModule.cpp b/Engine/Source/Runtime/Landscape/Private/LandscapeModule.cpp
--- a/Engine/Source/Runtime/Landscape/Private/LandscapeModule.cpp
+++ b/Engine/Source/Runtime/Landscape/Private/LandscapeModule.cpp
@@ -41,19 +41,18 @@ void LandscapeMaterialsParameterValuesGetter(FStaticParameterSet &OutStaticParam
 bool LandscapeMaterialsParameterSetUpdater(FStaticParameterSet &OutStaticParameterSet, UMaterial* Material);
 
 /**
- * Function that will fire everytime a world is created.
+ * Function that will fire every time a world is created.
  *
  * @param World A world that was created.
  * @param IVS Initialization values.
  */
-void WorldInitializationEventFunction(UWorld* World, const UWorld::InitializationValues IVS)
+void WorldCreationEventFunction(UWorld* World)
 {
 	AddPerWorldLandscapeData(World);
 }
 
-
 /**
- * Function that will fire everytime a world is destroyed.
+ * Function that will fire every time a world is destroyed.
  *
  * @param World A world that's being destroyed.
  */
@@ -95,8 +94,8 @@ void FLandscapeModule::StartupModule()
 	);
 
 #if WITH_EDITORONLY_DATA
-	FWorldDelegates::OnPostWorldInitialization.AddStatic(
-		&WorldInitializationEventFunction
+	FWorldDelegates::OnPostWorldCreation.AddStatic(
+		&WorldCreationEventFunction
 	);
 	FWorldDelegates::OnPreWorldFinishDestroy.AddStatic(
 		&WorldDestroyEventFunction
diff --git a/Engine/Source/Runtime/Launch/Private/LaunchEngineLoop.cpp b/Engine/Source/Runtime/Launch/Private/LaunchEngineLoop.cpp
--- a/Engine/Source/Runtime/Launch/Private/LaunchEngineLoop.cpp
+++ b/Engine/Source/Runtime/Launch/Private/LaunchEngineLoop.cpp
@@ -652,6 +652,40 @@ bool IsServerDelegateForOSS(FName WorldContextHandle)
 }
 #endif
 
+/**
+ * Loads array of module names that needs to be auto-loaded.
+ *
+ * @param OutModules Output array to fill with module names.
+ */
+void GetAutoStartupModuleList(TArray<FString>& OutModules)
+{
+	// UBT generated function for listing all auto-startup module
+	// this binary is loading. If the list is finished then this function
+	// returns nullptr.
+	extern const ANSICHAR* EnumAutoStartupModuleName(int Index);
+
+	const ANSICHAR* NamePtr = nullptr;
+	int32 Index = 0;
+	while ((NamePtr = EnumAutoStartupModuleName(Index++)) != nullptr)
+	{
+		OutModules.Add(FString(NamePtr));
+	}
+}
+
+/**
+ * Calls for each auto-startup module its StartupModule function.
+ */
+void InitializeAutoStartupModules()
+{
+	TArray<FString> Modules;
+	GetAutoStartupModuleList(Modules);
+
+	for (auto ModuleName : Modules)
+	{
+		FModuleManager::LoadModuleChecked<IModuleInterface>(*ModuleName);
+	}
+}
+
 int32 FEngineLoop::PreInit( const TCHAR* CmdLine )
 {
 	DECLARE_SCOPE_CYCLE_COUNTER(TEXT("Engine Pre-Initialized"), STAT_PreInit, STATGROUP_LoadTime);
@@ -1485,6 +1519,8 @@ int32 FEngineLoop::PreInit( const TCHAR* CmdLine )
 			// Execute the commandlet.
 			double CommandletExecutionStartTime = FPlatformTime::Seconds();
 
+			InitializeAutoStartupModules();
+
 			// Commandlets don't always handle -run= properly in the commandline so we'll provide them
 			// with a custom version that doesn't have it.
 			Commandlet->ParseParms( CommandletCommandLine );
@@ -1787,40 +1823,6 @@ bool FEngineLoop::LoadStartupCoreModules()
 	return bSuccess;
 }
 
-/**
- * Loads array of module names that needs to be auto-loaded.
- *
- * @param OutModules Output array to fill with module names.
- */
-void GetAutoStartupModuleList(TArray<FString>& OutModules)
-{
-	// UBT generated function for listing all auto-startup module
-	// this binary is loading. If the list is finished then this function
-	// returns nullptr.
-	extern const ANSICHAR* EnumAutoStartupModuleName(int Index);
-
-	const ANSICHAR* NamePtr = nullptr;
-	int32 Index = 0;
-	while ((NamePtr = EnumAutoStartupModuleName(Index++)) != nullptr)
-	{
-		OutModules.Add(FString(NamePtr));
-	}
-}
-
-/**
- * Calls for each auto-startup module its StartupModule function.
- */
-void InitializeAutoStartupModules()
-{
-	TArray<FString> Modules;
-	GetAutoStartupModuleList(Modules);
-
-	for (auto ModuleName : Modules)
-	{
-		FModuleManager::LoadModuleChecked<IModuleInterface>(*ModuleName);
-	}
-}
-
 bool FEngineLoop::LoadStartupModules()
 {
 	// Load any modules that want to be loaded before default modules are loaded up.
