Fixed monolithic builds

diff --git a/Engine/Source/Runtime/Launch/Launch.Build.cs b/Engine/Source/Runtime/Launch/Launch.Build.cs
--- a/Engine/Source/Runtime/Launch/Launch.Build.cs
+++ b/Engine/Source/Runtime/Launch/Launch.Build.cs
@@ -32,6 +32,7 @@ public class Launch : ModuleRules
 				"RenderCore",
 				"RHI",
 				"SandboxFile",
+				"Serialization",
 				"ShaderCore",
 				"Slate",
 				"SlateCore",
@@ -43,7 +44,7 @@ public class Launch : ModuleRules
 		{
 			PrivateDependencyModuleNames.AddRange(
 				new string[] {
-				"HeadMountedDisplay",
+					"HeadMountedDisplay",
 				}
 			);
 
@@ -70,7 +71,7 @@ public class Launch : ModuleRules
 		if (Target.Configuration != UnrealTargetConfiguration.Shipping)
 		{
 			PrivateDependencyModuleNames.AddRange(
-					new string[] {
+				new string[] {
 					"NetworkFile",
 					"StreamingFile",
     				"AutomationWorker",
@@ -193,7 +194,6 @@ public class Launch : ModuleRules
 			}
         }
 
-
 		// @todo ps4 clang bug: this works around a PS4/clang compiler bug (optimizations)
 		if (Target.Platform == UnrealTargetPlatform.PS4)
 		{
