Fix a JNI boolean call method
#android
#ue4
#codeview Josh.Adams,Terence.Burns

diff --git a/Engine/Source/Runtime/Launch/Private/Android/AndroidJNI.cpp b/Engine/Source/Runtime/Launch/Private/Android/AndroidJNI.cpp
--- a/Engine/Source/Runtime/Launch/Private/Android/AndroidJNI.cpp
+++ b/Engine/Source/Runtime/Launch/Private/Android/AndroidJNI.cpp
@@ -328,7 +328,7 @@ bool AndroidThunkCpp_IsMusicActive()
 	bool active = false;
 	if (JNIEnv* Env = GetJavaEnv())
 	{
-		active = (bool)Env->CallObjectMethod(GJavaGlobalThis, JDef_GameActivity::AndroidThunkJava_IsMusicActive);
+		active = (bool)Env->CallBooleanMethod(GJavaGlobalThis, JDef_GameActivity::AndroidThunkJava_IsMusicActive);
 	}
 	
 	return active;
