Fixed an issue with NSNotFount not matching our comparison index. Was a 32 v 64 bit comparison issue,

UE-4927

diff --git a/Engine/Source/Runtime/Launch/Private/IOS/IOSAppDelegateConsoleHandling.cpp b/Engine/Source/Runtime/Launch/Private/IOS/IOSAppDelegateConsoleHandling.cpp
--- a/Engine/Source/Runtime/Launch/Private/IOS/IOSAppDelegateConsoleHandling.cpp
+++ b/Engine/Source/Runtime/Launch/Private/IOS/IOSAppDelegateConsoleHandling.cpp
@@ -63,7 +63,7 @@
 			new(GEngine->DeferredCommands) FString(Ch.GetData());
 		}
 
-		uint32 ExistingCommand = [self.ConsoleHistoryValues indexOfObjectPassingTest:
+		NSUInteger ExistingCommand = [self.ConsoleHistoryValues indexOfObjectPassingTest:
 			^ BOOL (id obj, NSUInteger idx, BOOL *stop)
 			{ 
 				return [obj caseInsensitiveCompare:ConsoleCommand] == NSOrderedSame; 
