Made sure log directory is created even if we are not using log files as this is required for minidump output.

diff --git a/Engine/Source/Runtime/Launch/Private/LaunchEngineLoop.cpp b/Engine/Source/Runtime/Launch/Private/LaunchEngineLoop.cpp
--- a/Engine/Source/Runtime/Launch/Private/LaunchEngineLoop.cpp
+++ b/Engine/Source/Runtime/Launch/Private/LaunchEngineLoop.cpp
@@ -2437,6 +2437,10 @@ bool FEngineLoop::AppInit( )
 #endif // !UE_BUILD_SHIPPING
 
 #if PLATFORM_WINDOWS
+
+	// make sure that the log directory exists
+	IFileManager::Get().MakeDirectory( *FPaths::GameLogDir() );
+
 	// update the mini dump filename now that we have enough info to point it to the log folder even in installed builds
 	FCString::Strcpy(MiniDumpFilenameW, *IFileManager::Get().ConvertToAbsolutePathForExternalAppForWrite(*FString::Printf(TEXT("%sunreal-v%i-%s.dmp"), *FPaths::GameLogDir(), GEngineVersion.GetChangelist(), *FDateTime::Now().ToString())));
 #endif
