Fixes the issue where we might tick slate without the game viewport active within the main loop.

diff --git a/Engine/Source/Runtime/Engine/Private/GameEngine.cpp b/Engine/Source/Runtime/Engine/Private/GameEngine.cpp
--- a/Engine/Source/Runtime/Engine/Private/GameEngine.cpp
+++ b/Engine/Source/Runtime/Engine/Private/GameEngine.cpp
@@ -300,7 +300,7 @@ TSharedRef<SWindow> UGameEngine::CreateGameWindow()
 
 void UGameEngine::SwitchGameWindowToUseGameViewport()
 {
-	if ( ensure(GameViewportWindow.IsValid()) && GameViewportWindow.Pin()->GetContent() != GameViewportWidget )
+	if (GameViewportWindow.IsValid() && GameViewportWindow.Pin()->GetContent() != GameViewportWidget)
 	{
 		if( !GameViewportWidget.IsValid() )
 		{
diff --git a/Engine/Source/Runtime/Launch/Private/LaunchEngineLoop.cpp b/Engine/Source/Runtime/Launch/Private/LaunchEngineLoop.cpp
--- a/Engine/Source/Runtime/Launch/Private/LaunchEngineLoop.cpp
+++ b/Engine/Source/Runtime/Launch/Private/LaunchEngineLoop.cpp
@@ -2120,10 +2120,8 @@ bool FEngineLoop::ShouldUseIdleMode() const
 
 void FEngineLoop::Tick()
 {
-	// If a movie that is blocking the game thread has been playing,
-	// wait for it to finish before we tick again, otherwise we'll get
-	// multiple threads ticking simultaneously, which is bad
-	GetMoviePlayer()->WaitForMovieToFinish();
+	// Ensure we aren't starting a frame while loading or playing a loading movie
+	ensure(GetMoviePlayer()->IsLoadingFinished() && !GetMoviePlayer()->IsMovieCurrentlyPlaying());
 
 	// early in the Tick() to get the callbacks for cvar changes called
 	IConsoleManager::Get().CallAllConsoleVariableSinks();
@@ -2208,7 +2206,12 @@ void FEngineLoop::Tick()
 		}
 
 		GEngine->Tick( FApp::GetDeltaTime(), bIdleMode );
-
+		
+		// If a movie that is blocking the game thread has been playing,
+		// wait for it to finish before we continue to tick or tick again
+		// We do this right after GEngine->Tick() because that is where user code would initiate a load / movie.
+		GetMoviePlayer()->WaitForMovieToFinish();
+		
 		if (GShaderCompilingManager)
 		{
 			// Process any asynchronous shader compile results that are ready, limit execution time
