Disable crash handler; getting better stack traces from system logcat for now
#android
#ue4

diff --git a/Engine/Source/Runtime/Launch/Private/Android/AndroidJNI.cpp b/Engine/Source/Runtime/Launch/Private/Android/AndroidJNI.cpp
--- a/Engine/Source/Runtime/Launch/Private/Android/AndroidJNI.cpp
+++ b/Engine/Source/Runtime/Launch/Private/Android/AndroidJNI.cpp
@@ -528,7 +528,8 @@ JNIEXPORT jint JNI_OnLoad(JavaVM* InJavaVM, void* InReserved)
 	if( !FPlatformMisc::IsDebuggerPresent() || GAlwaysReportCrash )
 #endif
 	{
-		FPlatformMisc::SetCrashHandler(EngineCrashHandler);
+		// disable crash handler.. getting better stack traces from system logcat for now
+//		FPlatformMisc::SetCrashHandler(EngineCrashHandler);
 	}
 
 	// Cache path to external storage
