Removing stale GameName ini entries. Changed the mismatched gamename warning to better reflect the recent GameName ini setting removal.

diff --git a/Engine/Source/Runtime/Launch/Private/LaunchEngineLoop.cpp b/Engine/Source/Runtime/Launch/Private/LaunchEngineLoop.cpp
--- a/Engine/Source/Runtime/Launch/Private/LaunchEngineLoop.cpp
+++ b/Engine/Source/Runtime/Launch/Private/LaunchEngineLoop.cpp
@@ -334,9 +334,9 @@ void LaunchFixGameNameCase()
 			else
 			{
 				const FText Message = FText::Format(
-					NSLOCTEXT("Core", "MismatchedGameNames", "The name of the .uproject file ('{0}') must match the GameName key in the [URL] section of Config/DefaultEngine.ini (currently '{1}').  Please either change the ini or rename the .uproject to match each other (case-insensitive match)."),
-					FText::FromString(FApp::GetGameName()),
-					FText::FromString(GameName));
+					NSLOCTEXT("Core", "MismatchedGameNames", "The name of the .uproject file ('{0}') must match the name of the project passed in the command line ('{1}')."),
+					FText::FromString(*GameName),
+					FText::FromString(FApp::GetGameName()));
 				if (!GIsBuildMachine)
 				{
 					UE_LOG(LogInit, Warning, TEXT("%s"), *Message.ToString());
