Make sure touch events don't called GetScreenRect before OpenGL initialized (would cause init from wrong thread)
#android
#ue4

diff --git a/Engine/Source/Runtime/Launch/Private/Android/LaunchAndroid.cpp b/Engine/Source/Runtime/Launch/Private/Android/LaunchAndroid.cpp
--- a/Engine/Source/Runtime/Launch/Private/Android/LaunchAndroid.cpp
+++ b/Engine/Source/Runtime/Launch/Private/Android/LaunchAndroid.cpp
@@ -483,6 +483,8 @@ void android_main(struct android_app* state)
 	AndroidMain(state);
 }
 
+extern bool GAndroidGPUInfoReady;
+
 //Called from the event process thread
 static int32_t HandleInputCB(struct android_app* app, AInputEvent* event)
 {
@@ -554,6 +556,11 @@ static int32_t HandleInputCB(struct android_app* app, AInputEvent* event)
 				FAndroidWindow::CalculateSurfaceSize(Window, Width, Height);
 			}
 
+			// make sure OpenGL context created before accepting touch events.. FAndroidWindow::GetScreenRect() may try to create it early from wrong thread if this is the first call
+			if (!GAndroidGPUInfoReady)
+			{
+				return 0;
+			}
 			FPlatformRect ScreenRect = FAndroidWindow::GetScreenRect();
 
 			if(isActionTargeted)
diff --git a/Engine/Source/Runtime/OpenGLDrv/Private/Android/AndroidOpenGLPrivate.h b/Engine/Source/Runtime/OpenGLDrv/Private/Android/AndroidOpenGLPrivate.h
--- a/Engine/Source/Runtime/OpenGLDrv/Private/Android/AndroidOpenGLPrivate.h
+++ b/Engine/Source/Runtime/OpenGLDrv/Private/Android/AndroidOpenGLPrivate.h
@@ -5,6 +5,7 @@
 =============================================================================*/
 #pragma once
 
+bool GAndroidGPUInfoReady = false;
 
 class FAndroidGPUInfo
 {
@@ -93,4 +94,5 @@ private:
 
 #endif
 	}
+	GAndroidGPUInfoReady = true;
 };
