Fix Landscape not loaded in time for shader serialization

Merging

//depot/UE4-UT/...

to //depot/UE4/...

diff --git a/Engine/Source/Runtime/Launch/Private/LaunchEngineLoop.cpp b/Engine/Source/Runtime/Launch/Private/LaunchEngineLoop.cpp
--- a/Engine/Source/Runtime/Launch/Private/LaunchEngineLoop.cpp
+++ b/Engine/Source/Runtime/Launch/Private/LaunchEngineLoop.cpp
@@ -1705,11 +1705,13 @@ void FEngineLoop::LoadPreInitModules()
 		FModuleManager::Get().LoadModuleChecked<ISlateRHIRendererModule>("SlateRHIRenderer");
 	}
 #endif
-	
+
+	FModuleManager::Get().LoadModule(TEXT("Landscape"));
+
 	// Initialize ShaderCore before loading or compiling any shaders,
 	// But after Renderer and any other modules which implement shader types.
 	FModuleManager::Get().LoadModule(TEXT("ShaderCore"));
-
+	
 #if WITH_EDITORONLY_DATA
 	// Load the texture compressor module before any textures load. They may
 	// compress asynchronously and that can lead to a race condition.
