Allow commandlets to be run from game executables

diff --git a/Engine/Source/Runtime/Core/Private/Misc/Core.cpp b/Engine/Source/Runtime/Core/Private/Misc/Core.cpp
--- a/Engine/Source/Runtime/Core/Private/Misc/Core.cpp
+++ b/Engine/Source/Runtime/Core/Private/Misc/Core.cpp
@@ -90,10 +90,13 @@ bool GIsReconstructingBlueprintInstances = false;
 /** Force blueprints to not compile on load */
 bool GForceDisableBlueprintCompileOnLoad = false;
 
+#if WITH_ENGINE
+bool					PRIVATE_GIsRunningCommandlet	= false;					/* Whether this executable is running a commandlet (custom command-line processing code) */
+#endif
+
 #if WITH_EDITORONLY_DATA
 bool					GIsEditor						= false;					/* Whether engine was launched for editing */
 bool					GIsImportingT3D					= false;					/* Whether editor is importing T3D */
-bool					PRIVATE_GIsRunningCommandlet	= false;					/* Whether this executable is running a commandlet (custom command-line processing code in an editor-like environment) */
 bool					GIsUCCMakeStandaloneHeaderGenerator = false;					/* Are we rebuilding script via the standalone header generator? */
 bool					GIsTransacting					= false;					/* true if there is an undo/redo operation in progress. */
 bool					GIntraFrameDebuggingGameThread	= false;					/* Indicates that the game thread is currently paused deep in a call stack; do not process any game thread tasks */
diff --git a/Engine/Source/Runtime/Core/Public/CoreGlobals.h b/Engine/Source/Runtime/Core/Public/CoreGlobals.h
--- a/Engine/Source/Runtime/Core/Public/CoreGlobals.h
+++ b/Engine/Source/Runtime/Core/Public/CoreGlobals.h
@@ -56,6 +56,9 @@ extern CORE_API bool GForceDisableBlueprintCompileOnLoad;
 /** Helper function to flush resource streaming. */
 extern CORE_API void(*GFlushStreamingFunc)(void);
 
+#if WITH_ENGINE
+extern CORE_API bool PRIVATE_GIsRunningCommandlet;
+#endif
 
 #if WITH_EDITORONLY_DATA
 
@@ -65,7 +68,6 @@ extern CORE_API void(*GFlushStreamingFunc)(void);
 */
 extern CORE_API bool GIsEditor;
 extern CORE_API bool GIsImportingT3D;
-extern CORE_API bool PRIVATE_GIsRunningCommandlet;
 extern CORE_API bool GIsUCCMakeStandaloneHeaderGenerator;
 extern CORE_API bool GIsTransacting;
 
@@ -79,7 +81,6 @@ extern CORE_API bool			GFirstFrameIntraFrameDebugging;
 #else
 
 #define GIsEditor								false
-#define PRIVATE_GIsRunningCommandlet			false
 #define GIsUCCMakeStandaloneHeaderGenerator		false
 #define GIntraFrameDebuggingGameThread			false
 #define GFirstFrameIntraFrameDebugging			false
@@ -92,7 +93,7 @@ extern CORE_API bool			GFirstFrameIntraFrameDebugging;
 */
 FORCEINLINE bool IsRunningCommandlet()
 {
-#if WITH_EDITORONLY_DATA
+#if WITH_ENGINE
 	return PRIVATE_GIsRunningCommandlet;
 #else
 	return false;
diff --git a/Engine/Source/Runtime/Launch/Private/LaunchEngineLoop.cpp b/Engine/Source/Runtime/Launch/Private/LaunchEngineLoop.cpp
--- a/Engine/Source/Runtime/Launch/Private/LaunchEngineLoop.cpp
+++ b/Engine/Source/Runtime/Launch/Private/LaunchEngineLoop.cpp
@@ -782,7 +782,7 @@ int32 FEngineLoop::PreInit( const TCHAR* CmdLine )
 
 	FString Token				= FParse::Token( ParsedCmdLine, 0);
 
-#if	UE_EDITOR
+#if WITH_ENGINE
 	TArray<FString> Tokens;
 	TArray<FString> Switches;
 	UCommandlet::ParseCommandLine(CommandLineCopy, Tokens, Switches);
@@ -944,7 +944,23 @@ int32 FEngineLoop::PreInit( const TCHAR* CmdLine )
 	}
 
 	bHasEditorToken = !bIsNotEditor;
-#else	//UE_EDITOR
+#elif WITH_ENGINE
+	const TCHAR* CommandletCommandLine = NULL;
+	if (bHasCommandletToken)
+	{
+#if STATS
+		FThreadStats::MasterDisableForever();
+#endif
+		if (Token.StartsWith(TEXT("run=")))
+		{
+			Token = Token.RightChop(4);
+			if (!Token.EndsWith(TEXT("Commandlet")))
+			{
+				Token += TEXT("Commandlet");
+			}
+		}
+		CommandletCommandLine = ParsedCmdLine;
+	}
 #if WITH_EDITOR && WITH_EDITORONLY_DATA
 	// If a non-editor target build w/ WITH_EDITOR and WITH_EDITORONLY_DATA, use the old token check...
 	//@todo. Is this something we need to support?
@@ -1113,8 +1129,8 @@ int32 FEngineLoop::PreInit( const TCHAR* CmdLine )
 		GIsServer = true;
 #if WITH_EDITOR
 		GIsEditor = true;
-		PRIVATE_GIsRunningCommandlet = true;
 #endif	//WITH_EDITOR
+		PRIVATE_GIsRunningCommandlet = true;
 
 		// We need to disregard the empty token as we try finding Token + "Commandlet" which would result in finding the
 		// UCommandlet class if Token is empty.
@@ -1154,8 +1170,8 @@ int32 FEngineLoop::PreInit( const TCHAR* CmdLine )
 			GIsServer = false;
 #if WITH_EDITORONLY_DATA
 			GIsEditor = false;
-			PRIVATE_GIsRunningCommandlet = false;
 #endif
+			PRIVATE_GIsRunningCommandlet = false;
 		}
 	}
 
@@ -1163,8 +1179,8 @@ int32 FEngineLoop::PreInit( const TCHAR* CmdLine )
 	{
 		GIsClient = false;
 		GIsServer = true;
-#if WITH_EDITOR
 		PRIVATE_GIsRunningCommandlet = false;
+#if WITH_EDITOR
 		GIsEditor = false;
 #endif
 		bIsSeekFreeDedicatedServer = FPlatformProperties::RequiresCookedData();
@@ -1391,7 +1407,6 @@ int32 FEngineLoop::PreInit( const TCHAR* CmdLine )
 
 	SlowTask.EnterProgressFrame(50);
 
-#if WITH_EDITOR
 	if (!bHasEditorToken)
 	{
 		UClass* CommandletClass = NULL;
@@ -1445,7 +1460,16 @@ int32 FEngineLoop::PreInit( const TCHAR* CmdLine )
 
 			GIsClient = Default->IsClient;
 			GIsServer = Default->IsServer;
+#if WITH_EDITOR
 			GIsEditor = Default->IsEditor;
+#else
+			if (Default->IsEditor)
+			{
+				UE_LOG(LogInit, Error, TEXT("Cannot run editor commandlet %s with game executable."), *CommandletClass->GetFullName());
+				GIsRequestingExit = true;
+				return 1;
+			}
+#endif
 			PRIVATE_GIsRunningCommandlet = true;
 			// Reset aux log if we don't want to log to the console window.
 			if( !Default->LogToConsole )
@@ -1459,6 +1483,7 @@ int32 FEngineLoop::PreInit( const TCHAR* CmdLine )
 			CommandletClass->GetDefaultObject<UCommandlet>()->CreateCustomEngine(CommandletCommandLine);
 			if ( GEngine == NULL )
 			{
+#if WITH_EDITOR
 				if ( GIsEditor )
 				{
 					UClass* EditorEngineClass = StaticLoadClass( UEditorEngine::StaticClass(), NULL, TEXT("engine-ini:/Script/Engine.Engine.EditorEngine"), NULL, LOAD_None, NULL );
@@ -1472,6 +1497,7 @@ int32 FEngineLoop::PreInit( const TCHAR* CmdLine )
 					UE_LOG(LogInit, Log, TEXT("Initializing Editor Engine Completed"));
 				}
 				else
+#endif
 				{
 					UClass* EngineClass = StaticLoadClass( UEngine::StaticClass(), NULL, TEXT("engine-ini:/Script/Engine.Engine.GameEngine"), NULL, LOAD_None, NULL );
 
@@ -1609,8 +1635,6 @@ int32 FEngineLoop::PreInit( const TCHAR* CmdLine )
 		}
 	}
 
-#endif	//WITH_EDITOR
-
 	// exit if wanted.
 	if( GIsRequestingExit )
 	{
