Fix compile issues on console platforms

diff --git a/Engine/Source/Runtime/Launch/Private/LaunchEngineLoop.cpp b/Engine/Source/Runtime/Launch/Private/LaunchEngineLoop.cpp
--- a/Engine/Source/Runtime/Launch/Private/LaunchEngineLoop.cpp
+++ b/Engine/Source/Runtime/Launch/Private/LaunchEngineLoop.cpp
@@ -1425,6 +1425,7 @@ int32 FEngineLoop::PreInit( const TCHAR* CmdLine )
 				return 1;
 			}
 
+#if PLATFORM_WINDOWS || PLATFORM_MAC || PLATFORM_LINUX
 			extern bool GIsConsoleExecutable;
 			if (GIsConsoleExecutable)
 			{
@@ -1436,6 +1437,7 @@ int32 FEngineLoop::PreInit( const TCHAR* CmdLine )
 				FPlatformMisc::SetGracefulTerminationHandler();
 			}
 			else
+#endif
 			{
 				// Bring up console unless we're a silent build.
 				if( GLogConsole && !GIsSilent )
