[UE-4242] Implement app lifecycle delegates on Android. The delegates will fire on these events: onResume() fires ApplicationHasEnteredForeground and ApplicationHasReactivated, onPause() fires ApplicationWillDeactivate and ApplicationWillEnterBackground, and onDestroy() fires ApplicationWillTerminate.

diff --git a/Engine/Source/Runtime/Launch/Private/Android/AndroidEventManager.cpp b/Engine/Source/Runtime/Launch/Private/Android/AndroidEventManager.cpp
--- a/Engine/Source/Runtime/Launch/Private/Android/AndroidEventManager.cpp
+++ b/Engine/Source/Runtime/Launch/Private/Android/AndroidEventManager.cpp
@@ -4,6 +4,7 @@
 #include "AndroidEventManager.h"
 #include "AndroidApplication.h"
 #include "AudioDevice.h"
+#include "CallbackDevice.h"
 #include <android/native_window.h> 
 #include <android/native_window_jni.h> 
 
@@ -56,16 +57,21 @@ void FAppEventManager::Tick()
 			//doing nothing here
 			break;
 		case APP_EVENT_STATE_ON_DESTROY:
+			FCoreDelegates::ApplicationWillTerminateDelegate.Broadcast();
 			GIsRequestingExit = true; //destroy immediately. Game will shutdown.
 			break;
 		case APP_EVENT_STATE_ON_STOP:
 			bHaveGame = false;
 			break;
 		case APP_EVENT_STATE_ON_PAUSE:
+			FCoreDelegates::ApplicationWillDeactivateDelegate.Broadcast();
+			FCoreDelegates::ApplicationWillEnterBackgroundDelegate.Broadcast();
 			bHaveGame = false;
 			break;
 		case APP_EVENT_STATE_ON_RESUME:
 			bHaveGame = true;
+			FCoreDelegates::ApplicationHasEnteredForegroundDelegate.Broadcast();
+			FCoreDelegates::ApplicationHasReactivatedDelegate.Broadcast();
 			break;
 
 		// window focus events that follow their own  heirarchy, and might or might not respect App main events heirarchy
