- Cleaned up the Android JNI code
- Started on the work for a separate GooglePlay java class, but that's not in yet (notice: GoogleServicesThis = FJavaWrapper::GameActivityThis)

diff --git a/Engine/Build/Android/Java/src/com/epicgames/ue4/GameActivity.java b/Engine/Build/Android/Java/src/com/epicgames/ue4/GameActivity.java
--- a/Engine/Build/Android/Java/src/com/epicgames/ue4/GameActivity.java
+++ b/Engine/Build/Android/Java/src/com/epicgames/ue4/GameActivity.java
@@ -896,15 +896,11 @@ public class GameActivity extends NativeActivity
 
 	public native void nativeConsoleCommand(String commandString);
 	public native void nativeVirtualKeyboardResult(boolean update, String contents);
-	
-	public native boolean nativeIsGooglePlayEnabled();
 
 	public native void nativeInitHMDs();
 
 	public native void nativeResumeMainInit();
 	
-	public native void nativeCompletedConnection(int playerID, int errorCode);
-
 	static
 	{
 		System.loadLibrary("gnustl_shared");
diff --git a/Engine/Source/Runtime/Launch/Private/Android/AndroidJNI.cpp b/Engine/Source/Runtime/Launch/Private/Android/AndroidJNI.cpp
--- a/Engine/Source/Runtime/Launch/Private/Android/AndroidJNI.cpp
+++ b/Engine/Source/Runtime/Launch/Private/Android/AndroidJNI.cpp
@@ -12,7 +12,6 @@
 #define JNI_CURRENT_VERSION JNI_VERSION_1_6
 
 JavaVM* GJavaVM;
-jobject GJavaGlobalThis = NULL;
 
 // Pointer to target widget for virtual keyboard contents
 static IVirtualKeyboardEntry *VirtualKeyboardWidget = NULL;
@@ -24,34 +23,156 @@ extern FString GOBBFilename;
 
 //////////////////////////////////////////////////////////////////////////
 
+#if UE_BUILD_SHIPPING
+// always clear any exceptions in SHipping
+#define CHECK_JNI_RESULT( Id ) if (Id == 0) Env->ExceptionClear();
+#else
+#define CHECK_JNI_RESULT( Id ) \
+if (Id == 0) \
+{ \
+if (bIsOptional) { Env->ExceptionClear(); } \
+	else { checkf(Id != 0, TEXT("JNI_OnLoad: Failed to find " #Id)); } \
+}
+#endif
+
+void FJavaWrapper::FindClassesAndMethods(JNIEnv* Env)
+{
+	bool bIsOptional = false;
+	GameActivityClassID = FindClass(Env, "com/epicgames/ue4/GameActivity", bIsOptional);
+	AndroidThunkJava_ShowConsoleWindow = FindMethod(Env, GameActivityClassID, "AndroidThunkJava_ShowConsoleWindow", "(Ljava/lang/String;)V", bIsOptional);
+	AndroidThunkJava_ShowVirtualKeyboardInput = FindMethod(Env, GameActivityClassID, "AndroidThunkJava_ShowVirtualKeyboardInput", "(ILjava/lang/String;Ljava/lang/String;)V", bIsOptional);
+	AndroidThunkJava_LaunchURL = FindMethod(Env, GameActivityClassID, "AndroidThunkJava_LaunchURL", "(Ljava/lang/String;)V", bIsOptional);
+	AndroidThunkJava_GetAssetManager = FindMethod(Env, GameActivityClassID, "AndroidThunkJava_GetAssetManager", "()Landroid/content/res/AssetManager;", bIsOptional);
+	AndroidThunkJava_Minimize = FindMethod(Env, GameActivityClassID, "AndroidThunkJava_Minimize", "()V", bIsOptional);
+	AndroidThunkJava_ForceQuit = FindMethod(Env, GameActivityClassID, "AndroidThunkJava_ForceQuit", "()V", bIsOptional);
+	AndroidThunkJava_GetFontDirectory = FindStaticMethod(Env, GameActivityClassID, "AndroidThunkJava_GetFontDirectory", "()Ljava/lang/String;", bIsOptional);
+	AndroidThunkJava_Vibrate = FindMethod(Env, GameActivityClassID, "AndroidThunkJava_Vibrate", "(J)V", bIsOptional);
+	AndroidThunkJava_IsMusicActive = FindMethod(Env, GameActivityClassID, "AndroidThunkJava_IsMusicActive", "()Z", bIsOptional);
+	AndroidThunkJava_KeepScreenOn = FindMethod(Env, GameActivityClassID, "AndroidThunkJava_KeepScreenOn", "(Z)V", bIsOptional);
+	AndroidThunkJava_InitHMDs = FindMethod(Env, GameActivityClassID, "AndroidThunkJava_InitHMDs", "()V", bIsOptional);
+
+
+	// the rest are optional
+	bIsOptional = true;
+	// @todo split GooglePlay
+	//	GoogleServicesClassID = FindClass(Env, "com/epicgames/ue4/GoogleServices", bIsOptional);
+	GoogleServicesClassID = GameActivityClassID;
+	AndroidThunkJava_ResetAchievements = FindMethod(Env, GoogleServicesClassID, "AndroidThunkJava_ResetAchievements", "()V", bIsOptional);
+	AndroidThunkJava_ShowAdBanner = FindMethod(Env, GoogleServicesClassID, "AndroidThunkJava_ShowAdBanner", "(Ljava/lang/String;Z)V", bIsOptional);
+	AndroidThunkJava_HideAdBanner = FindMethod(Env, GoogleServicesClassID, "AndroidThunkJava_HideAdBanner", "()V", bIsOptional);
+	AndroidThunkJava_CloseAdBanner = FindMethod(Env, GoogleServicesClassID, "AndroidThunkJava_CloseAdBanner", "()V", bIsOptional);
+
+	// In app purchase functionality
+	AndroidThunkJava_IapSetupService = FindMethod(Env, GoogleServicesClassID, "AndroidThunkJava_IapSetupService", "(Ljava/lang/String;)V", bIsOptional);
+	AndroidThunkJava_IapQueryInAppPurchases = FindMethod(Env, GoogleServicesClassID, "AndroidThunkJava_IapQueryInAppPurchases", "([Ljava/lang/String;[Z)Z", bIsOptional);
+	AndroidThunkJava_IapBeginPurchase = FindMethod(Env, GoogleServicesClassID, "AndroidThunkJava_IapBeginPurchase", "(Ljava/lang/String;Z)Z", bIsOptional);
+	AndroidThunkJava_IapIsAllowedToMakePurchases = FindMethod(Env, GoogleServicesClassID, "AndroidThunkJava_IapIsAllowedToMakePurchases", "()Z", bIsOptional);
+}
+
+jclass FJavaWrapper::FindClass(JNIEnv* Env, const ANSICHAR* ClassName, bool bIsOptional)
+{
+	jclass Class = Env->FindClass(ClassName);
+	CHECK_JNI_RESULT(Class);
+	return Class;
+}
+
+jmethodID FJavaWrapper::FindMethod(JNIEnv* Env, jclass Class, const ANSICHAR* MethodName, const ANSICHAR* MethodSignature, bool bIsOptional)
+{
+	jmethodID Method = Class == NULL ? NULL : Env->GetMethodID(Class, MethodName, MethodSignature);
+	CHECK_JNI_RESULT(Method);
+	return Method;
+}
+
+jmethodID FJavaWrapper::FindStaticMethod(JNIEnv* Env, jclass Class, const ANSICHAR* MethodName, const ANSICHAR* MethodSignature, bool bIsOptional)
+{
+	jmethodID Method = Class == NULL ? NULL : Env->GetStaticMethodID(Class, MethodName, MethodSignature);
+	CHECK_JNI_RESULT(Method);
+	return Method;
+}
+
+jfieldID FJavaWrapper::FindField(JNIEnv* Env, jclass Class, const ANSICHAR* FieldName, const ANSICHAR* FieldType, bool bIsOptional)
+{
+	jfieldID Field = Class == NULL ? NULL : Env->GetFieldID(Class, FieldName, FieldType);
+	CHECK_JNI_RESULT(Field);
+	return Field;
+}
+
+void FJavaWrapper::CallVoidMethod(JNIEnv* Env, jobject Object, jmethodID Method, ...)
+{
+	// make sure the function exists
+	if (Method == NULL || Object == NULL)
+	{
+		return;
+	}
+
+	va_list Args;
+	va_start(Args, Method);
+	Env->CallVoidMethod(Object, Method, Args);
+	va_end(Args);
+}
+
+jobject FJavaWrapper::CallObjectMethod(JNIEnv* Env, jobject Object, jmethodID Method, ...)
+{
+	if (Method == NULL || Object == NULL)
+	{
+		return NULL;
+	}
+
+	va_list Args;
+	va_start(Args, Method);
+	jobject Return = Env->CallObjectMethod(Object, Method, Args);
+	va_end(Args);
+
+	return Return;
+}
+
+bool FJavaWrapper::CallBooleanMethod(JNIEnv* Env, jobject Object, jmethodID Method, ...)
+{
+	if (Method == NULL || Object == NULL)
+	{
+		return false;
+	}
+
+	va_list Args;
+	va_start(Args, Method);
+	jboolean Return = Env->CallBooleanMethod(Object, Method, Args);
+	va_end(Args);
+
+	return (bool)Return;
+}
+
 //Declare all the static members of the class defs 
-jclass JDef_GameActivity::ClassID;
-jmethodID JDef_GameActivity::AndroidThunkJava_KeepScreenOn;
-jmethodID JDef_GameActivity::AndroidThunkJava_Vibrate;
-jmethodID JDef_GameActivity::AndroidThunkJava_ShowConsoleWindow;
-jmethodID JDef_GameActivity::AndroidThunkJava_ShowVirtualKeyboardInput;
-jmethodID JDef_GameActivity::AndroidThunkJava_LaunchURL;
-jmethodID JDef_GameActivity::AndroidThunkJava_ResetAchievements;
-jmethodID JDef_GameActivity::AndroidThunkJava_ShowAdBanner;
-jmethodID JDef_GameActivity::AndroidThunkJava_HideAdBanner;
-jmethodID JDef_GameActivity::AndroidThunkJava_CloseAdBanner;
-jmethodID JDef_GameActivity::AndroidThunkJava_GetAssetManager;
-jmethodID JDef_GameActivity::AndroidThunkJava_Minimize;
-jmethodID JDef_GameActivity::AndroidThunkJava_ForceQuit;
-jmethodID JDef_GameActivity::AndroidThunkJava_GetFontDirectory;
-jmethodID JDef_GameActivity::AndroidThunkJava_IsMusicActive;
-jmethodID JDef_GameActivity::AndroidThunkJava_InitHMDs;
-jmethodID JDef_GameActivity::AndroidThunkJava_IapSetupService;
-jmethodID JDef_GameActivity::AndroidThunkJava_IapQueryInAppPurchases;
-jmethodID JDef_GameActivity::AndroidThunkJava_IapBeginPurchase;
-jmethodID JDef_GameActivity::AndroidThunkJava_IapIsAllowedToMakePurchases;
+jclass FJavaWrapper::GameActivityClassID;
+jobject FJavaWrapper::GameActivityThis;
+jmethodID FJavaWrapper::AndroidThunkJava_ShowConsoleWindow;
+jmethodID FJavaWrapper::AndroidThunkJava_ShowVirtualKeyboardInput;
+jmethodID FJavaWrapper::AndroidThunkJava_LaunchURL;
+jmethodID FJavaWrapper::AndroidThunkJava_GetAssetManager;
+jmethodID FJavaWrapper::AndroidThunkJava_Minimize;
+jmethodID FJavaWrapper::AndroidThunkJava_ForceQuit;
+jmethodID FJavaWrapper::AndroidThunkJava_GetFontDirectory;
+jmethodID FJavaWrapper::AndroidThunkJava_Vibrate;
+jmethodID FJavaWrapper::AndroidThunkJava_IsMusicActive;
+jmethodID FJavaWrapper::AndroidThunkJava_KeepScreenOn;
+jmethodID FJavaWrapper::AndroidThunkJava_InitHMDs;
+
+jclass FJavaWrapper::GoogleServicesClassID;
+jobject FJavaWrapper::GoogleServicesThis;
+jmethodID FJavaWrapper::AndroidThunkJava_ResetAchievements;
+jmethodID FJavaWrapper::AndroidThunkJava_ShowAdBanner;
+jmethodID FJavaWrapper::AndroidThunkJava_HideAdBanner;
+jmethodID FJavaWrapper::AndroidThunkJava_CloseAdBanner;
+jmethodID FJavaWrapper::AndroidThunkJava_IapSetupService;
+jmethodID FJavaWrapper::AndroidThunkJava_IapQueryInAppPurchases;
+jmethodID FJavaWrapper::AndroidThunkJava_IapBeginPurchase;
+jmethodID FJavaWrapper::AndroidThunkJava_IapIsAllowedToMakePurchases;
 
 DEFINE_LOG_CATEGORY_STATIC(LogEngine, Log, All);
 
 //Game-specific crash reporter
 void EngineCrashHandler(const FGenericCrashContext& GenericContext)
 {
-	const FAndroidCrashContext& Context = static_cast< const FAndroidCrashContext& >( GenericContext );
+	const FAndroidCrashContext& Context = static_cast<const FAndroidCrashContext&>(GenericContext);
 
 	static int32 bHasEntered = 0;
 	if (FPlatformAtomics::InterlockedCompareExchange(&bHasEntered, 1, 0) == 0)
@@ -75,7 +196,7 @@ void AndroidThunkCpp_KeepScreenOn(bool Enable)
 	if (JNIEnv* Env = FAndroidApplication::GetJavaEnv())
 	{
 		// call the java side
-		Env->CallVoidMethod(GJavaGlobalThis, JDef_GameActivity::AndroidThunkJava_KeepScreenOn, Enable);
+		FJavaWrapper::CallVoidMethod(Env, FJavaWrapper::GameActivityThis, FJavaWrapper::AndroidThunkJava_KeepScreenOn, Enable);
 	}
 }
 
@@ -84,7 +205,7 @@ void AndroidThunkCpp_Vibrate(int64_t Duration)
 	if (JNIEnv* Env = FAndroidApplication::GetJavaEnv())
 	{
 		// call the java side
-		Env->CallVoidMethod(GJavaGlobalThis, JDef_GameActivity::AndroidThunkJava_Vibrate, Duration);
+		FJavaWrapper::CallVoidMethod(Env, FJavaWrapper::GameActivityThis, FJavaWrapper::AndroidThunkJava_Vibrate, Duration);
 	}
 }
 
@@ -93,7 +214,7 @@ void AndroidThunkCpp_InitHMDs()
 {
 	if (JNIEnv* Env = FAndroidApplication::GetJavaEnv())
 	{
-		Env->CallVoidMethod(GJavaGlobalThis, JDef_GameActivity::AndroidThunkJava_InitHMDs);
+		FJavaWrapper::CallVoidMethod(Env, FJavaWrapper::GameActivityThis, FJavaWrapper::AndroidThunkJava_InitHMDs);
 	}
 }
 
@@ -117,14 +238,14 @@ void AndroidThunkCpp_ShowConsoleWindow()
 				{
 					ConsoleText += ", ";
 				}
-	
-				ConsoleText += Format.Mid(UnderscoreIndex+1);
+
+				ConsoleText += Format.Mid(UnderscoreIndex + 1);
 			}
 		}
 
 		// call the java side
-		jstring ConsoleTextJava	= Env->NewStringUTF(TCHAR_TO_UTF8(*ConsoleText));
-		Env->CallVoidMethod(GJavaGlobalThis, JDef_GameActivity::AndroidThunkJava_ShowConsoleWindow, ConsoleTextJava);
+		jstring ConsoleTextJava = Env->NewStringUTF(TCHAR_TO_UTF8(*ConsoleText));
+		FJavaWrapper::CallVoidMethod(Env, FJavaWrapper::GameActivityThis, FJavaWrapper::AndroidThunkJava_ShowConsoleWindow, ConsoleTextJava);
 		Env->DeleteLocalRef(ConsoleTextJava);
 	}
 }
@@ -139,7 +260,7 @@ void AndroidThunkCpp_ShowVirtualKeyboardInput(TSharedPtr<IVirtualKeyboardEntry>
 		// call the java side
 		jstring LabelJava = Env->NewStringUTF(TCHAR_TO_UTF8(*Label));
 		jstring ContentsJava = Env->NewStringUTF(TCHAR_TO_UTF8(*Contents));
-		Env->CallVoidMethod(GJavaGlobalThis, JDef_GameActivity::AndroidThunkJava_ShowVirtualKeyboardInput, InputType, LabelJava, ContentsJava);
+		FJavaWrapper::CallVoidMethod(Env, FJavaWrapper::GameActivityThis, FJavaWrapper::AndroidThunkJava_ShowVirtualKeyboardInput, InputType, LabelJava, ContentsJava);
 		Env->DeleteLocalRef(ContentsJava);
 		Env->DeleteLocalRef(LabelJava);
 	}
@@ -170,7 +291,7 @@ void AndroidThunkCpp_LaunchURL(const FString& URL)
 	if (JNIEnv* Env = FAndroidApplication::GetJavaEnv())
 	{
 		jstring Argument = Env->NewStringUTF(TCHAR_TO_UTF8(*URL));
-		Env->CallVoidMethod(GJavaGlobalThis, JDef_GameActivity::AndroidThunkJava_LaunchURL, Argument);
+		FJavaWrapper::CallVoidMethod(Env, FJavaWrapper::GameActivityThis, FJavaWrapper::AndroidThunkJava_LaunchURL, Argument);
 		Env->DeleteLocalRef(Argument);
 	}
 }
@@ -179,7 +300,7 @@ void AndroidThunkCpp_ResetAchievements()
 {
 	if (JNIEnv* Env = FAndroidApplication::GetJavaEnv())
 	{
-		Env->CallVoidMethod(GJavaGlobalThis, JDef_GameActivity::AndroidThunkJava_ResetAchievements);
+		FJavaWrapper::CallVoidMethod(Env, FJavaWrapper::GoogleServicesThis, FJavaWrapper::AndroidThunkJava_ResetAchievements);
 	}
 }
 
@@ -188,16 +309,16 @@ void AndroidThunkCpp_ShowAdBanner(const FString& AdUnitID, bool bShowOnBottomOfS
 	if (JNIEnv* Env = FAndroidApplication::GetJavaEnv())
  	{
 		jstring AdUnitIDArg = Env->NewStringUTF(TCHAR_TO_UTF8(*AdUnitID));
-		Env->CallVoidMethod(GJavaGlobalThis, JDef_GameActivity::AndroidThunkJava_ShowAdBanner, AdUnitIDArg, bShowOnBottomOfScreen);
+		FJavaWrapper::CallVoidMethod(Env, FJavaWrapper::GoogleServicesThis, FJavaWrapper::AndroidThunkJava_ShowAdBanner, AdUnitIDArg, bShowOnBottomOfScreen);
 		Env->DeleteLocalRef(AdUnitIDArg);
- 	}
+	}
 }
 
 void AndroidThunkCpp_HideAdBanner()
 {
  	if (JNIEnv* Env = FAndroidApplication::GetJavaEnv())
  	{
-		Env->CallVoidMethod(GJavaGlobalThis, JDef_GameActivity::AndroidThunkJava_HideAdBanner);
+		FJavaWrapper::CallVoidMethod(Env, FJavaWrapper::GoogleServicesThis, FJavaWrapper::AndroidThunkJava_HideAdBanner);
  	}
 }
 
@@ -205,7 +326,7 @@ void AndroidThunkCpp_CloseAdBanner()
 {
  	if (JNIEnv* Env = FAndroidApplication::GetJavaEnv())
  	{
-		Env->CallVoidMethod(GJavaGlobalThis, JDef_GameActivity::AndroidThunkJava_CloseAdBanner);
+		FJavaWrapper::CallVoidMethod(Env, FJavaWrapper::GoogleServicesThis, FJavaWrapper::AndroidThunkJava_CloseAdBanner);
  	}
 }
 
@@ -221,7 +342,7 @@ jobject AndroidJNI_GetJavaAssetManager()
 	{
 		if (JNIEnv* Env = FAndroidApplication::GetJavaEnv())
 		{
-			GJavaAssetManager = Env->CallObjectMethod(GJavaGlobalThis, JDef_GameActivity::AndroidThunkJava_GetAssetManager);
+			GJavaAssetManager = FJavaWrapper::CallObjectMethod(Env, FJavaWrapper::GameActivityThis, FJavaWrapper::AndroidThunkJava_GetAssetManager);
 			Env->NewGlobalRef(GJavaAssetManager);
 		}
 	}
@@ -246,7 +367,7 @@ void AndroidThunkCpp_Minimize()
 {
 	if (JNIEnv* Env = FAndroidApplication::GetJavaEnv())
 	{
-		Env->CallVoidMethod(GJavaGlobalThis, JDef_GameActivity::AndroidThunkJava_Minimize);
+		FJavaWrapper::CallVoidMethod(Env, FJavaWrapper::GameActivityThis, FJavaWrapper::AndroidThunkJava_Minimize);
 	}
 }
 
@@ -254,35 +375,28 @@ void AndroidThunkCpp_ForceQuit()
 {
 	if (JNIEnv* Env = FAndroidApplication::GetJavaEnv())
 	{
-		Env->CallVoidMethod(GJavaGlobalThis, JDef_GameActivity::AndroidThunkJava_ForceQuit);
+		FJavaWrapper::CallVoidMethod(Env, FJavaWrapper::GameActivityThis, FJavaWrapper::AndroidThunkJava_ForceQuit);
 	}
 }
 
 bool AndroidThunkCpp_IsMusicActive()
 {
-	bool active = false;
+	bool bIsActive = false;
 	if (JNIEnv* Env = FAndroidApplication::GetJavaEnv())
 	{
-		active = (bool)Env->CallBooleanMethod(GJavaGlobalThis, JDef_GameActivity::AndroidThunkJava_IsMusicActive);
+		bIsActive = FJavaWrapper::CallBooleanMethod(Env, FJavaWrapper::GameActivityThis, FJavaWrapper::AndroidThunkJava_IsMusicActive);
 	}
-	
-	return active;
+
+	return bIsActive;
 }
 
 void AndroidThunkCpp_Iap_SetupIapService(const FString& InProductKey)
 {
 	if (JNIEnv* Env = FAndroidApplication::GetJavaEnv())
 	{
-		if (GJavaGlobalThis)
-		{
-			jstring ProductKey = Env->NewStringUTF(TCHAR_TO_UTF8(*InProductKey));
-			Env->CallVoidMethod(GJavaGlobalThis, JDef_GameActivity::AndroidThunkJava_IapSetupService, ProductKey);
-			Env->DeleteLocalRef(ProductKey);
-		}
-		else
-		{
-			FPlatformMisc::LowLevelOutputDebugString(L"[JNI] - JNI is invalid");
-		}
+		jstring ProductKey = Env->NewStringUTF(TCHAR_TO_UTF8(*InProductKey));
+		FJavaWrapper::CallVoidMethod(Env, FJavaWrapper::GoogleServicesThis, FJavaWrapper::AndroidThunkJava_IapSetupService, ProductKey);
+		Env->DeleteLocalRef(ProductKey);
 	}
 }
 
@@ -293,36 +407,29 @@ bool AndroidThunkCpp_Iap_QueryInAppPurchases(const TArray<FString>& ProductIDs,
 
 	if (JNIEnv* Env = FAndroidApplication::GetJavaEnv())
 	{
-		if (GJavaGlobalThis)
-		{
-			static jclass StringClass = Env->FindClass("java/lang/String");
-
-			// Populate some java types with the provided product information
-			jobjectArray ProductIDArray = (jobjectArray)Env->NewObjectArray(ProductIDs.Num(), StringClass, NULL);
-			jbooleanArray ConsumeArray = (jbooleanArray)Env->NewBooleanArray(ProductIDs.Num());
-
-			jboolean* ConsumeArrayValues = Env->GetBooleanArrayElements(ConsumeArray, 0);
-			for (uint32 Param = 0; Param < ProductIDs.Num(); Param++)
-			{
-				jstring StringValue = Env->NewStringUTF(TCHAR_TO_UTF8(*ProductIDs[Param]));
-				Env->SetObjectArrayElement(ProductIDArray, Param, StringValue);
-				Env->DeleteLocalRef(StringValue);
+		static jclass StringClass = Env->FindClass("java/lang/String");
 
-				ConsumeArrayValues[Param] = bConsumable[Param];
-			}
-			Env->ReleaseBooleanArrayElements(ConsumeArray, ConsumeArrayValues, 0);
-
-			// Execute the java code for this operation
-			bResult = (bool)Env->CallBooleanMethod(GJavaGlobalThis, JDef_GameActivity::AndroidThunkJava_IapQueryInAppPurchases, ProductIDArray, ConsumeArray);
+		// Populate some java types with the provided product information
+		jobjectArray ProductIDArray = (jobjectArray)Env->NewObjectArray(ProductIDs.Num(), StringClass, NULL);
+		jbooleanArray ConsumeArray = (jbooleanArray)Env->NewBooleanArray(ProductIDs.Num());
 
-			// clean up references
-			Env->DeleteLocalRef(ProductIDArray);
-			Env->DeleteLocalRef(ConsumeArray);
-		}
-		else
+		jboolean* ConsumeArrayValues = Env->GetBooleanArrayElements(ConsumeArray, 0);
+		for (uint32 Param = 0; Param < ProductIDs.Num(); Param++)
 		{
-			FPlatformMisc::LowLevelOutputDebugString(L"[JNI] - JNI is invalid");
+			jstring StringValue = Env->NewStringUTF(TCHAR_TO_UTF8(*ProductIDs[Param]));
+			Env->SetObjectArrayElement(ProductIDArray, Param, StringValue);
+			Env->DeleteLocalRef(StringValue);
+
+			ConsumeArrayValues[Param] = bConsumable[Param];
 		}
+		Env->ReleaseBooleanArrayElements(ConsumeArray, ConsumeArrayValues, 0);
+
+		// Execute the java code for this operation
+		bResult = FJavaWrapper::CallBooleanMethod(Env, FJavaWrapper::GoogleServicesThis, FJavaWrapper::AndroidThunkJava_IapQueryInAppPurchases, ProductIDArray, ConsumeArray);
+		
+		// clean up references
+		Env->DeleteLocalRef(ProductIDArray);
+		Env->DeleteLocalRef(ConsumeArray);
 	}
 
 	return bResult;
@@ -334,18 +441,9 @@ bool AndroidThunkCpp_Iap_BeginPurchase(const FString& ProductID, const bool bCon
 	bool bResult = false;
 	if (JNIEnv* Env = FAndroidApplication::GetJavaEnv())
 	{
-		if (GJavaGlobalThis)
-		{
-			jstring ProductIdJava = Env->NewStringUTF(TCHAR_TO_UTF8(*ProductID));
-			bResult = (bool)Env->CallBooleanMethod(GJavaGlobalThis, JDef_GameActivity::AndroidThunkJava_IapBeginPurchase, ProductIdJava, bConsumable);
-			Env->DeleteLocalRef(ProductIdJava);
-
-			FPlatformMisc::LowLevelOutputDebugStringf(TEXT("JDef_GameActivity::AndroidThunkJava_IapBeginPurchase - '%s'"), bResult ? TEXT("true") : TEXT("false"));
-		}
-		else
-		{
-			FPlatformMisc::LowLevelOutputDebugString(L"[JNI] - JNI is invalid");
-		}
+		jstring ProductIdJava = Env->NewStringUTF(TCHAR_TO_UTF8(*ProductID));
+		bResult = FJavaWrapper::CallBooleanMethod(Env, FJavaWrapper::GoogleServicesThis, FJavaWrapper::AndroidThunkJava_IapBeginPurchase, ProductIdJava, bConsumable);
+		Env->DeleteLocalRef(ProductIdJava);
 	}
 
 	return bResult;
@@ -357,14 +455,7 @@ bool AndroidThunkCpp_Iap_IsAllowedToMakePurchases()
 	bool bResult = false;
 	if (JNIEnv* Env = FAndroidApplication::GetJavaEnv())
 	{
-		if (GJavaGlobalThis)
-		{
-			bResult = (bool)Env->CallBooleanMethod(GJavaGlobalThis, JDef_GameActivity::AndroidThunkJava_IapIsAllowedToMakePurchases);
-		}
-		else
-		{
-			FPlatformMisc::LowLevelOutputDebugString(L"[JNI] - JNI is invalid");
-		}
+		bResult = FJavaWrapper::CallBooleanMethod(Env, FJavaWrapper::GoogleServicesThis, FJavaWrapper::AndroidThunkJava_IapIsAllowedToMakePurchases);
 	}
 	return bResult;
 }
@@ -381,83 +472,25 @@ bool AndroidThunkCpp_Iap_IsAllowedToMakePurchases()
 // variables. That way, later when the Java functions need to be called, the IDs will be ready.
 // It is much slower to keep looking up the class and method IDs.
 
-#if UE_BUILD_SHIPPING
-#define CHECK_JNI_RESULT( Id )
-#else
-#define CHECK_JNI_RESULT( Id ) \
-	if ( Id == 0 ) \
-	{ \
-		FPlatformMisc::LowLevelOutputDebugString(TEXT("JNI_OnLoad: Failed to find " #Id)); \
-	}
-#endif
-
 JNIEXPORT jint JNI_OnLoad(JavaVM* InJavaVM, void* InReserved)
 {
 	FPlatformMisc::LowLevelOutputDebugString(L"In the JNI_OnLoad function");
 
-	JNIEnv* env = NULL;
-	InJavaVM->GetEnv((void **)&env, JNI_CURRENT_VERSION);
+	JNIEnv* Env = NULL;
+	InJavaVM->GetEnv((void **)&Env, JNI_CURRENT_VERSION);
 
 	// if you have problems with stuff being missing esspecially in distribution builds then it could be because proguard is stripping things from java
 	// check proguard-project.txt and see if your stuff is included in the exceptions
 	GJavaVM = InJavaVM;
-	FAndroidApplication::InitializeJavaEnv(GJavaVM, JNI_CURRENT_VERSION, GJavaGlobalThis);
-
-	JDef_GameActivity::ClassID = env->FindClass("com/epicgames/ue4/GameActivity");
-	CHECK_JNI_RESULT( JDef_GameActivity::ClassID );
-
-	JDef_GameActivity::AndroidThunkJava_KeepScreenOn = env->GetMethodID(JDef_GameActivity::ClassID, "AndroidThunkJava_KeepScreenOn", "(Z)V");
-	CHECK_JNI_RESULT(JDef_GameActivity::AndroidThunkJava_KeepScreenOn);
-	JDef_GameActivity::AndroidThunkJava_Vibrate = env->GetMethodID(JDef_GameActivity::ClassID, "AndroidThunkJava_Vibrate", "(J)V");
-	CHECK_JNI_RESULT(JDef_GameActivity::AndroidThunkJava_Vibrate);
-	JDef_GameActivity::AndroidThunkJava_ShowConsoleWindow = env->GetMethodID(JDef_GameActivity::ClassID, "AndroidThunkJava_ShowConsoleWindow", "(Ljava/lang/String;)V");
-	CHECK_JNI_RESULT( JDef_GameActivity::AndroidThunkJava_ShowConsoleWindow );
-	JDef_GameActivity::AndroidThunkJava_ShowVirtualKeyboardInput = env->GetMethodID(JDef_GameActivity::ClassID, "AndroidThunkJava_ShowVirtualKeyboardInput", "(ILjava/lang/String;Ljava/lang/String;)V");
-	CHECK_JNI_RESULT( JDef_GameActivity::AndroidThunkJava_ShowVirtualKeyboardInput );
-	JDef_GameActivity::AndroidThunkJava_LaunchURL = env->GetMethodID(JDef_GameActivity::ClassID, "AndroidThunkJava_LaunchURL", "(Ljava/lang/String;)V");
-	CHECK_JNI_RESULT( JDef_GameActivity::AndroidThunkJava_LaunchURL );
-	JDef_GameActivity::AndroidThunkJava_ResetAchievements = env->GetMethodID(JDef_GameActivity::ClassID, "AndroidThunkJava_ResetAchievements", "()V");
-	JDef_GameActivity::AndroidThunkJava_ShowAdBanner = env->GetMethodID(JDef_GameActivity::ClassID, "AndroidThunkJava_ShowAdBanner", "(Ljava/lang/String;Z)V");
-	CHECK_JNI_RESULT( JDef_GameActivity::AndroidThunkJava_ShowAdBanner );
-	JDef_GameActivity::AndroidThunkJava_HideAdBanner = env->GetMethodID(JDef_GameActivity::ClassID, "AndroidThunkJava_HideAdBanner", "()V");
-	CHECK_JNI_RESULT( JDef_GameActivity::AndroidThunkJava_HideAdBanner );
-	JDef_GameActivity::AndroidThunkJava_CloseAdBanner = env->GetMethodID(JDef_GameActivity::ClassID, "AndroidThunkJava_CloseAdBanner", "()V");
-	CHECK_JNI_RESULT( JDef_GameActivity::AndroidThunkJava_CloseAdBanner );
-	JDef_GameActivity::AndroidThunkJava_GetAssetManager = env->GetMethodID(JDef_GameActivity::ClassID, "AndroidThunkJava_GetAssetManager", "()Landroid/content/res/AssetManager;");
-	CHECK_JNI_RESULT( JDef_GameActivity::AndroidThunkJava_GetAssetManager );
-
-	JDef_GameActivity::AndroidThunkJava_Minimize = env->GetMethodID(JDef_GameActivity::ClassID, "AndroidThunkJava_Minimize", "()V");
-	CHECK_JNI_RESULT( JDef_GameActivity::AndroidThunkJava_Minimize );
-	JDef_GameActivity::AndroidThunkJava_ForceQuit = env->GetMethodID(JDef_GameActivity::ClassID, "AndroidThunkJava_ForceQuit", "()V");
-	CHECK_JNI_RESULT( JDef_GameActivity::AndroidThunkJava_ForceQuit );
-	
-	JDef_GameActivity::AndroidThunkJava_GetFontDirectory = env->GetStaticMethodID(JDef_GameActivity::ClassID, "AndroidThunkJava_GetFontDirectory", "()Ljava/lang/String;");
-	CHECK_JNI_RESULT( JDef_GameActivity::AndroidThunkJava_GetFontDirectory );
-	
-	JDef_GameActivity::AndroidThunkJava_IsMusicActive = env->GetMethodID(JDef_GameActivity::ClassID, "AndroidThunkJava_IsMusicActive", "()Z");
-	CHECK_JNI_RESULT( JDef_GameActivity::AndroidThunkJava_IsMusicActive );	
-
-	JDef_GameActivity::AndroidThunkJava_InitHMDs = env->GetMethodID(JDef_GameActivity::ClassID, "AndroidThunkJava_InitHMDs", "()V");
-	CHECK_JNI_RESULT(JDef_GameActivity::AndroidThunkJava_InitHMDs);
-
-	// In app purchase functionality
-	JDef_GameActivity::AndroidThunkJava_IapSetupService = env->GetMethodID(JDef_GameActivity::ClassID, "AndroidThunkJava_IapSetupService", "(Ljava/lang/String;)V");
-	CHECK_JNI_RESULT(JDef_GameActivity::AndroidThunkJava_IapSetupService);
-
-	JDef_GameActivity::AndroidThunkJava_IapQueryInAppPurchases = env->GetMethodID(JDef_GameActivity::ClassID, "AndroidThunkJava_IapQueryInAppPurchases", "([Ljava/lang/String;[Z)Z");
-	CHECK_JNI_RESULT(JDef_GameActivity::AndroidThunkJava_IapQueryInAppPurchases);
-
-	JDef_GameActivity::AndroidThunkJava_IapBeginPurchase = env->GetMethodID(JDef_GameActivity::ClassID, "AndroidThunkJava_IapBeginPurchase", "(Ljava/lang/String;Z)Z");
-	CHECK_JNI_RESULT(JDef_GameActivity::AndroidThunkJava_IapBeginPurchase);
+	FAndroidApplication::InitializeJavaEnv(GJavaVM, JNI_CURRENT_VERSION, FJavaWrapper::GameActivityThis);
 
-	JDef_GameActivity::AndroidThunkJava_IapIsAllowedToMakePurchases = env->GetMethodID(JDef_GameActivity::ClassID, "AndroidThunkJava_IapIsAllowedToMakePurchases", "()Z");
-	CHECK_JNI_RESULT(JDef_GameActivity::AndroidThunkJava_IapIsAllowedToMakePurchases);
+	FJavaWrapper::FindClassesAndMethods(Env);
 
 	// hook signals
 #if UE_BUILD_DEBUG
-	if( GAlwaysReportCrash )
+	if (GAlwaysReportCrash)
 #else
-	if( !FPlatformMisc::IsDebuggerPresent() || GAlwaysReportCrash )
+	if (0)//!FPlatformMisc::IsDebuggerPresent() || GAlwaysReportCrash)
 #endif
 	{
 		// disable crash handler.. getting better stack traces from system logcat for now
@@ -465,35 +498,35 @@ JNIEXPORT jint JNI_OnLoad(JavaVM* InJavaVM, void* InReserved)
 	}
 
 	// Cache path to external storage
-	jclass EnvClass = env->FindClass("android/os/Environment");
-	jmethodID getExternalStorageDir = env->GetStaticMethodID(EnvClass, "getExternalStorageDirectory", "()Ljava/io/File;");
-	jobject externalStoragePath = env->CallStaticObjectMethod(EnvClass, getExternalStorageDir, nullptr);
-	jmethodID getFilePath = env->GetMethodID(env->FindClass("java/io/File"), "getPath", "()Ljava/lang/String;");
-	jstring pathString = (jstring)env->CallObjectMethod(externalStoragePath, getFilePath, nullptr);
-	const char *nativePathString = env->GetStringUTFChars(pathString, 0);
+	jclass EnvClass = Env->FindClass("android/os/Environment");
+	jmethodID getExternalStorageDir = Env->GetStaticMethodID(EnvClass, "getExternalStorageDirectory", "()Ljava/io/File;");
+	jobject externalStoragePath = Env->CallStaticObjectMethod(EnvClass, getExternalStorageDir, nullptr);
+	jmethodID getFilePath = Env->GetMethodID(Env->FindClass("java/io/File"), "getPath", "()Ljava/lang/String;");
+	jstring pathString = (jstring)Env->CallObjectMethod(externalStoragePath, getFilePath, nullptr);
+	const char *nativePathString = Env->GetStringUTFChars(pathString, 0);
 	// Copy that somewhere safe 
 	GFilePathBase = FString(nativePathString);
-	
+
 	// then release...
-	env->ReleaseStringUTFChars(pathString, nativePathString);
+	Env->ReleaseStringUTFChars(pathString, nativePathString);
 	FPlatformMisc::LowLevelOutputDebugStringf(TEXT("Path found as '%s'\n"), *GFilePathBase);
 
 	// Next we check to see if the OBB file is in the APK
-	jmethodID isOBBInAPKMethod = env->GetStaticMethodID(JDef_GameActivity::ClassID, "isOBBInAPK", "()Z");
-	GOBBinAPK = (bool)env->CallStaticBooleanMethod(JDef_GameActivity::ClassID, isOBBInAPKMethod, nullptr);
+	jmethodID isOBBInAPKMethod = Env->GetStaticMethodID(FJavaWrapper::GameActivityClassID, "isOBBInAPK", "()Z");
+	GOBBinAPK = (bool)Env->CallStaticBooleanMethod(FJavaWrapper::GameActivityClassID, isOBBInAPKMethod, nullptr);
 
 	// Get the OBB filename
-	jmethodID GetOBBFilenameMethod = env->GetStaticMethodID(JDef_GameActivity::ClassID, "GetOBBFilename", "()Ljava/lang/String;");
-	jstring obbString = (jstring)env->CallStaticObjectMethod(JDef_GameActivity::ClassID, GetOBBFilenameMethod, nullptr);
-	const char * nativeObbString = env->GetStringUTFChars(obbString, 0);
-	GOBBFilename = FString(nativeObbString);
-	env->ReleaseStringUTFChars(obbString, nativeObbString);
+	jmethodID GetOBBFilenameMethod = Env->GetStaticMethodID(FJavaWrapper::GameActivityClassID, "GetOBBFilename", "()Ljava/lang/String;");
+	jstring ObbString = (jstring)Env->CallStaticObjectMethod(FJavaWrapper::GameActivityClassID, GetOBBFilenameMethod, nullptr);
+	const char * NativeObbString = Env->GetStringUTFChars(ObbString, 0);
+	GOBBFilename = FString(NativeObbString);
+	Env->ReleaseStringUTFChars(ObbString, NativeObbString);
 
 	// Get the system font directory
-	jstring fontPath = (jstring)env->CallStaticObjectMethod(JDef_GameActivity::ClassID, JDef_GameActivity::AndroidThunkJava_GetFontDirectory);
-	const char * nativeFontPathString = env->GetStringUTFChars(fontPath, 0);
+	jstring fontPath = (jstring)Env->CallStaticObjectMethod(FJavaWrapper::GameActivityClassID, FJavaWrapper::AndroidThunkJava_GetFontDirectory);
+	const char * nativeFontPathString = Env->GetStringUTFChars(fontPath, 0);
 	GFontPathBase = FString(nativeFontPathString);
-	env->ReleaseStringUTFChars(fontPath, nativeFontPathString);
+	Env->ReleaseStringUTFChars(fontPath, nativeFontPathString);
 	FPlatformMisc::LowLevelOutputDebugStringf(TEXT("Font Path found as '%s'\n"), *GFontPathBase);
 
 	// Wire up to core delegates, so core code can call out to Java
@@ -510,17 +543,21 @@ JNIEXPORT jint JNI_OnLoad(JavaVM* InJavaVM, void* InReserved)
 //Native-defined functions
 
 //This function is declared in the Java-defined class, GameActivity.java: "public native void nativeSetGlobalActivity();"
-extern "C" void Java_com_epicgames_ue4_GameActivity_nativeSetGlobalActivity(JNIEnv* jenv, jobject thiz)
+extern "C" void Java_com_epicgames_ue4_GameActivity_nativeSetGlobalActivity(JNIEnv* jenv, jobject thiz /*, jobject googleServices*/)
 {
-	if (!GJavaGlobalThis)
+	if (!FJavaWrapper::GameActivityThis)
 	{
-		GJavaGlobalThis = jenv->NewGlobalRef(thiz);
-		FAndroidApplication::InitializeJavaEnv(GJavaVM, JNI_CURRENT_VERSION, GJavaGlobalThis);
-		if(!GJavaGlobalThis)
+		FJavaWrapper::GameActivityThis = jenv->NewGlobalRef(thiz);
+		FAndroidApplication::InitializeJavaEnv(GJavaVM, JNI_CURRENT_VERSION, FJavaWrapper::GameActivityThis);
+		if (!FJavaWrapper::GameActivityThis)
 		{
 			FPlatformMisc::LowLevelOutputDebugString(L"Error setting the global GameActivity activity");
 			check(false);
 		}
+
+		// @todo split GooglePlay, this needs to be passed in to this function
+		FJavaWrapper::GoogleServicesThis = FJavaWrapper::GameActivityThis;
+		// FJavaWrapper::GoogleServicesThis = jenv->NewGlobalRef(googleServices);
 	}
 }
 
diff --git a/Engine/Source/Runtime/Launch/Private/Android/AndroidJNI.h b/Engine/Source/Runtime/Launch/Private/Android/AndroidJNI.h
--- a/Engine/Source/Runtime/Launch/Private/Android/AndroidJNI.h
+++ b/Engine/Source/Runtime/Launch/Private/Android/AndroidJNI.h
@@ -6,41 +6,55 @@
 #include <android/log.h>
 
 extern JavaVM* GJavaVM;
-extern jobject GJavaGlobalThis;
-
-/*
- * NOTE -- At the moment, this is specific to GameActivity
-*/
 
 // Define all the Java classes/methods that the game will need to access to
-class JDef_GameActivity
+class FJavaWrapper
 {
 public:
-	static jclass ClassID;
 
 	// Nonstatic methods
-	static jmethodID AndroidThunkJava_KeepScreenOn;
-	static jmethodID AndroidThunkJava_Vibrate;
+	static jclass GameActivityClassID;
+	static jobject GameActivityThis;
 	static jmethodID AndroidThunkJava_ShowConsoleWindow;
 	static jmethodID AndroidThunkJava_ShowVirtualKeyboardInput;
 	static jmethodID AndroidThunkJava_LaunchURL;
-	static jmethodID AndroidThunkJava_ResetAchievements;
-	static jmethodID AndroidThunkJava_ShowAdBanner;
-	static jmethodID AndroidThunkJava_HideAdBanner;
-	static jmethodID AndroidThunkJava_CloseAdBanner;
 	static jmethodID AndroidThunkJava_GetAssetManager;
 	static jmethodID AndroidThunkJava_Minimize;
 	static jmethodID AndroidThunkJava_ForceQuit;
-
-	static jmethodID AndroidThunkJava_InitHMDs;
-
 	static jmethodID AndroidThunkJava_GetFontDirectory;
-
+	static jmethodID AndroidThunkJava_Vibrate;
 	static jmethodID AndroidThunkJava_IsMusicActive;
+	static jmethodID AndroidThunkJava_KeepScreenOn;
+	static jmethodID AndroidThunkJava_InitHMDs;
+
+	// IDs related to google play services
+	static jclass GoogleServicesClassID;
+	static jobject GoogleServicesThis;
+	static jmethodID AndroidThunkJava_ResetAchievements;
+	static jmethodID AndroidThunkJava_ShowAdBanner;
+	static jmethodID AndroidThunkJava_HideAdBanner;
+	static jmethodID AndroidThunkJava_CloseAdBanner;
 
 	// In app purchase functionality
 	static jmethodID AndroidThunkJava_IapSetupService;
 	static jmethodID AndroidThunkJava_IapQueryInAppPurchases;
 	static jmethodID AndroidThunkJava_IapBeginPurchase;
 	static jmethodID AndroidThunkJava_IapIsAllowedToMakePurchases;
+
+	/**
+	 * Find all known classes and methods
+	 */
+	static void FindClassesAndMethods(JNIEnv* Env);
+
+	/**
+	 * Helper wrapper functions around the JNIEnv versions with NULL/error handling
+	 */
+	static jclass FindClass(JNIEnv* Env, const ANSICHAR* ClassName, bool bIsOptional);
+	static jmethodID FindMethod(JNIEnv* Env, jclass Class, const ANSICHAR* MethodName, const ANSICHAR* MethodSignature, bool bIsOptional);
+	static jmethodID FindStaticMethod(JNIEnv* Env, jclass Class, const ANSICHAR* MethodName, const ANSICHAR* MethodSignature, bool bIsOptional);
+	static jfieldID FindField(JNIEnv* Env, jclass Class, const ANSICHAR* FieldName, const ANSICHAR* FieldType, bool bIsOptional);
+
+	static void CallVoidMethod(JNIEnv* Env, jobject Object, jmethodID Method, ...);
+	static jobject CallObjectMethod(JNIEnv* Env, jobject Object, jmethodID Method, ...);
+	static bool CallBooleanMethod(JNIEnv* Env, jobject Object, jmethodID Method, ...);
 };
diff --git a/Engine/Source/Runtime/Launch/Private/Android/LaunchAndroid.cpp b/Engine/Source/Runtime/Launch/Private/Android/LaunchAndroid.cpp
--- a/Engine/Source/Runtime/Launch/Private/Android/LaunchAndroid.cpp
+++ b/Engine/Source/Runtime/Launch/Private/Android/LaunchAndroid.cpp
@@ -979,15 +979,6 @@ extern "C" void Java_com_epicgames_ue4_GameActivity_nativeConsoleCommand(JNIEnv*
 	jenv->ReleaseStringUTFChars(commandString, javaChars);
 }
 
-//This function is declared in the Java-defined class, GameActivity.java: "public native void nativeIsGooglePlayEnabled();"
-extern "C" jboolean Java_com_epicgames_ue4_GameActivity_nativeIsGooglePlayEnabled(JNIEnv* jenv, jobject thiz)
-{
-	bool bEnabled = false;
-	GConfig->GetBool(TEXT("/Script/AndroidRuntimeSettings.AndroidRuntimeSettings"), TEXT("bEnableGooglePlaySupport"), bEnabled, GEngineIni);
-	UE_LOG(LogOnline, Log, TEXT("Checking whether Google Play is enabled. bEnableGooglePlaySupport = %d"), bEnabled);
-	return bEnabled;
-}
-
 // This is called from the Java UI thread for initializing VR HMDs
 extern "C" void Java_com_epicgames_ue4_GameActivity_nativeInitHMDs(JNIEnv* jenv, jobject thiz)
 {
diff --git a/Engine/Source/Runtime/Online/Android/OnlineSubsystemGooglePlay/Private/OnlineIdentityInterfaceGoolgePlay.cpp b/Engine/Source/Runtime/Online/Android/OnlineSubsystemGooglePlay/Private/OnlineIdentityInterfaceGoolgePlay.cpp
--- a/Engine/Source/Runtime/Online/Android/OnlineSubsystemGooglePlay/Private/OnlineIdentityInterfaceGoolgePlay.cpp
+++ b/Engine/Source/Runtime/Online/Android/OnlineSubsystemGooglePlay/Private/OnlineIdentityInterfaceGoolgePlay.cpp
@@ -2,7 +2,6 @@
 
 #include "OnlineSubsystemGooglePlayPrivatePCH.h"
 
-#include "Android/AndroidJNI.h"
 FOnlineIdentityGooglePlay::FPendingConnection FOnlineIdentityGooglePlay::PendingConnectRequest;
 
 FOnlineIdentityGooglePlay::FOnlineIdentityGooglePlay(FOnlineSubsystemGooglePlay* InSubsystem)
