Dead code removal in LaunchAndroid.cpp.

diff --git a/Engine/Source/Runtime/Launch/Private/Android/LaunchAndroid.cpp b/Engine/Source/Runtime/Launch/Private/Android/LaunchAndroid.cpp
--- a/Engine/Source/Runtime/Launch/Private/Android/LaunchAndroid.cpp
+++ b/Engine/Source/Runtime/Launch/Private/Android/LaunchAndroid.cpp
@@ -167,52 +167,6 @@ void InitHMDs()
 	}
 }
 
-
-void UpdateGameInterruptions()
-{
-	// Check for game suspension.
-	if(GHasInterruptionRequest)
-	{
-		// Suspend the renderer.
-		if(GUseThreadedRendering)
-		{
-			FlushRenderingCommands();
-			StopRenderingThread();
-		}
-		else
-		{
-			RHIReleaseThreadOwnership();
-		}
-
-		// Flag the suspended state.
-		GIsInterrupted = true;
-
-		// Wait for resume.
-		while(GHasInterruptionRequest)
-		{
-			FPlatformProcess::Sleep(0.1f);
-		}
-
-		// Flag the resume state.
-		GIsInterrupted = false;
-
-		// Reset the window surface.
-		RHIAcquireThreadOwnership();
-		RHIReleaseThreadOwnership();
-
-		// Resume the renderer.
-		if(GUseThreadedRendering)
-		{
-			StartRenderingThread();
-		}
-		else
-		{
-			UE_LOG(LogAndroid, Display, TEXT("Acquiring Thread Ownership"));
-			RHIAcquireThreadOwnership();
-		}
-	}
-}
-
 static void InitCommandLine()
 {
 	static const uint32 CMD_LINE_MAX = 16384u;
