- Improved CHECK_JNI_RESULT printout

diff --git a/Engine/Source/Runtime/Launch/Private/Android/AndroidJNI.cpp b/Engine/Source/Runtime/Launch/Private/Android/AndroidJNI.cpp
--- a/Engine/Source/Runtime/Launch/Private/Android/AndroidJNI.cpp
+++ b/Engine/Source/Runtime/Launch/Private/Android/AndroidJNI.cpp
@@ -24,13 +24,13 @@ extern bool GOBBinAPK;
 
 #if UE_BUILD_SHIPPING
 // always clear any exceptions in SHipping
-#define CHECK_JNI_RESULT( Id ) if (Id == 0) Env->ExceptionClear();
+#define CHECK_JNI_RESULT(Id) if (Id == 0) { Env->ExceptionClear(); }
 #else
-#define CHECK_JNI_RESULT( Id ) \
+#define CHECK_JNI_RESULT(Id) \
 if (Id == 0) \
 { \
-if (bIsOptional) { Env->ExceptionClear(); } \
-	else { checkf(Id != 0, TEXT("JNI_OnLoad: Failed to find " #Id)); } \
+	if (bIsOptional) { Env->ExceptionClear(); } \
+	else { Env->ExceptionDescribe(); checkf(Id != 0, TEXT("Failed to find " #Id)); } \
 }
 #endif
 
