Moved Android lifetime event delegates to fire off suspend before the rendering thread shuts down, and resume after the rendering thread starts up again.

diff --git a/Engine/Source/Runtime/Launch/Private/Android/AndroidEventManager.cpp b/Engine/Source/Runtime/Launch/Private/Android/AndroidEventManager.cpp
--- a/Engine/Source/Runtime/Launch/Private/Android/AndroidEventManager.cpp
+++ b/Engine/Source/Runtime/Launch/Private/Android/AndroidEventManager.cpp
@@ -64,14 +64,10 @@ void FAppEventManager::Tick()
 			bHaveGame = false;
 			break;
 		case APP_EVENT_STATE_ON_PAUSE:
-			FCoreDelegates::ApplicationWillDeactivateDelegate.Broadcast();
-			FCoreDelegates::ApplicationWillEnterBackgroundDelegate.Broadcast();
 			bHaveGame = false;
 			break;
 		case APP_EVENT_STATE_ON_RESUME:
 			bHaveGame = true;
-			FCoreDelegates::ApplicationHasEnteredForegroundDelegate.Broadcast();
-			FCoreDelegates::ApplicationHasReactivatedDelegate.Broadcast();
 			break;
 
 		// window focus events that follow their own  heirarchy, and might or might not respect App main events heirarchy
@@ -116,12 +112,22 @@ void FAppEventManager::Tick()
 		{
 			ResumeRendering();
 			ResumeAudio();
+
+			// broadcast events after the rendering thread has resumed
+			FCoreDelegates::ApplicationHasEnteredForegroundDelegate.Broadcast();
+			FCoreDelegates::ApplicationHasReactivatedDelegate.Broadcast();
+
 			bRunning = true;
 		}
 		else if (bRunning && (!bHaveWindow || !bHaveGame))
 		{
+			// broadcast events before rendering thred suspends
+			FCoreDelegates::ApplicationWillDeactivateDelegate.Broadcast();
+			FCoreDelegates::ApplicationWillEnterBackgroundDelegate.Broadcast();
+
 			PauseRendering();
 			PauseAudio();
+
 			bRunning = false;
 		}
 	}
