[UE-6973] Fix crash when doing four-finger touch gesture to bring up console. This fixes all activity JNI calls to use correct vaargs variants.

diff --git a/Engine/Source/Runtime/Launch/Private/Android/AndroidJNI.cpp b/Engine/Source/Runtime/Launch/Private/Android/AndroidJNI.cpp
--- a/Engine/Source/Runtime/Launch/Private/Android/AndroidJNI.cpp
+++ b/Engine/Source/Runtime/Launch/Private/Android/AndroidJNI.cpp
@@ -106,7 +106,7 @@ void FJavaWrapper::CallVoidMethod(JNIEnv* Env, jobject Object, jmethodID Method,
 
 	va_list Args;
 	va_start(Args, Method);
-	Env->CallVoidMethod(Object, Method, Args);
+	Env->CallVoidMethodV(Object, Method, Args);
 	va_end(Args);
 }
 
@@ -119,7 +119,7 @@ jobject FJavaWrapper::CallObjectMethod(JNIEnv* Env, jobject Object, jmethodID Me
 
 	va_list Args;
 	va_start(Args, Method);
-	jobject Return = Env->CallObjectMethod(Object, Method, Args);
+	jobject Return = Env->CallObjectMethodV(Object, Method, Args);
 	va_end(Args);
 
 	return Return;
@@ -134,7 +134,7 @@ bool FJavaWrapper::CallBooleanMethod(JNIEnv* Env, jobject Object, jmethodID Meth
 
 	va_list Args;
 	va_start(Args, Method);
-	jboolean Return = Env->CallBooleanMethod(Object, Method, Args);
+	jboolean Return = Env->CallBooleanMethodV(Object, Method, Args);
 	va_end(Args);
 
 	return (bool)Return;
