- delay android window destruction until after the renderer teardown, only on GearVR

diff --git a/Engine/Source/Runtime/Launch/Private/Android/AndroidEventManager.cpp b/Engine/Source/Runtime/Launch/Private/Android/AndroidEventManager.cpp
--- a/Engine/Source/Runtime/Launch/Private/Android/AndroidEventManager.cpp
+++ b/Engine/Source/Runtime/Launch/Private/Android/AndroidEventManager.cpp
@@ -29,6 +29,8 @@ void FAppEventManager::Tick()
 {
 	while (!Queue.IsEmpty())
 	{
+		bool bDestroyWindow = false;
+
 		FAppEventData Event = DequeueAppEvent();
 
 		switch (Event.State)
@@ -48,9 +50,16 @@ void FAppEventManager::Tick()
 			bSaveState = true; //todo android: handle save state.
 			break;
 		case APP_EVENT_STATE_WINDOW_DESTROYED:
-			//AndroidEGL::GetInstance()->UnBind()
-			FAndroidAppEntry::DestroyWindow();
-			FPlatformMisc::SetHardwareWindow(NULL);
+			if (GEngine->HMDDevice.IsValid() && GEngine->HMDDevice->IsHMDConnected())
+			{
+				// delay the destruction until after the renderer teardown on GearVR
+				bDestroyWindow = true;
+			}
+			else
+			{
+				FAndroidAppEntry::DestroyWindow();
+				FPlatformMisc::SetHardwareWindow(NULL);
+			}
 			bHaveWindow = false;
 			break;
 		case APP_EVENT_STATE_ON_START:
@@ -130,6 +139,13 @@ void FAppEventManager::Tick()
 
 			bRunning = false;
 		}
+
+		if (bDestroyWindow)
+		{
+			FAndroidAppEntry::DestroyWindow();
+			FPlatformMisc::SetHardwareWindow(NULL);
+			bDestroyWindow = false;
+		}
 	}
 
 	if (!bRunning && FirstInitialized)
