Fix crash in Xbox One Suspend()/Resume() handling.
- Suspend/resume is thread safe (so long as we're not issuing rendering commands), so we don't need the RHIFlush which will only run on the render thread. Changed the macro we use to expose the suspend/resume functions.
- Moved Xbox One specific suspend/resume handling code out of the generic EngineLoop.

diff --git a/Engine/Source/Runtime/Launch/Private/LaunchEngineLoop.cpp b/Engine/Source/Runtime/Launch/Private/LaunchEngineLoop.cpp
--- a/Engine/Source/Runtime/Launch/Private/LaunchEngineLoop.cpp
+++ b/Engine/Source/Runtime/Launch/Private/LaunchEngineLoop.cpp
@@ -2383,42 +2383,6 @@ void FEngineLoop::ClearPendingCleanupObjects()
 	PendingCleanupObjects = NULL;
 }
 
-#if PLATFORM_XBOXONE
-
-void FEngineLoop::OnResuming(_In_ Platform::Object^ Sender, _In_ Platform::Object^ Args)
-{
-	// Make the call down to the RHI to Resume the GPU state
-	RHIResumeRendering();
-
-	// Notify application of resume
-	FCoreDelegates::ApplicationHasEnteredForegroundDelegate.Broadcast();
-}
-
-void FEngineLoop::OnSuspending(_In_ Platform::Object^ Sender, _In_ Windows::ApplicationModel::SuspendingEventArgs^ Args)
-{
-	// Get the Suspending Event
-	Windows::ApplicationModel::SuspendingDeferral^ SuspendingEvent = Args->SuspendingOperation->GetDeferral();
-	
-	// Notify application of suspend. Application should kick off an async save at this point.
-	FCoreDelegates::ApplicationWillEnterBackgroundDelegate.Broadcast();
-
-	// Flush the RenderingThread
-	FlushRenderingCommands();
-
-	// Make the call down to the RHI to Suspend the GPU state
-	RHISuspendRendering();
-
-	// @TODO Wait for async save to complete
-	// Flush the log so it's all written to disk
-	GLog->FlushThreadedLogs();
-	GLog->Flush();
-
-	// Tell the callback that we are done
-	SuspendingEvent->Complete();
-}
-
-#endif // PLATFORM_XBOXONE
-
 #endif // WITH_ENGINE
 
 
diff --git a/Engine/Source/Runtime/Launch/Public/LaunchEngineLoop.h b/Engine/Source/Runtime/Launch/Public/LaunchEngineLoop.h
--- a/Engine/Source/Runtime/Launch/Public/LaunchEngineLoop.h
+++ b/Engine/Source/Runtime/Launch/Public/LaunchEngineLoop.h
@@ -78,12 +78,6 @@ public:
 	/** Removes references to any objects pending cleanup by deleting them. */
 	virtual void ClearPendingCleanupObjects() override;
 
-#if PLATFORM_XBOXONE
-	void OnResuming(_In_ Platform::Object^ Sender, _In_ Platform::Object^ Args);
-	void OnSuspending(_In_ Platform::Object^ Sender, _In_ Windows::ApplicationModel::SuspendingEventArgs^ Args);
-	void OnResourceAvailabilityChanged( _In_ Platform::Object^ Sender, _In_ Platform::Object^ Args );
-#endif // PLATFORM_XBOXONE
-
 #endif // WITH_ENGINE
 
 	/** Pre-init HMD device (if necessary). */
diff --git a/Engine/Source/Runtime/RHI/Public/RHIMethods.h b/Engine/Source/Runtime/RHI/Public/RHIMethods.h
--- a/Engine/Source/Runtime/RHI/Public/RHIMethods.h
+++ b/Engine/Source/Runtime/RHI/Public/RHIMethods.h
@@ -2066,11 +2066,11 @@ DEFINE_RHIMETHOD_0(
 	);
 
 // Operations to suspend title rendering and yield control to the system
-DEFINE_RHIMETHOD_GLOBAL_0(
+DEFINE_RHIMETHOD_GLOBALTHREADSAFE_0(
 	void,SuspendRendering,
 	,
 	);
-DEFINE_RHIMETHOD_GLOBAL_0(
+DEFINE_RHIMETHOD_GLOBALTHREADSAFE_0(
 	void,ResumeRendering,
 	,
 	);
