Fixed missed checkin
#android
#ue4

diff --git a/Engine/Source/Runtime/Launch/Private/Android/LaunchAndroid.cpp b/Engine/Source/Runtime/Launch/Private/Android/LaunchAndroid.cpp
--- a/Engine/Source/Runtime/Launch/Private/Android/LaunchAndroid.cpp
+++ b/Engine/Source/Runtime/Launch/Private/Android/LaunchAndroid.cpp
@@ -960,3 +960,8 @@ extern "C" void Java_com_epicgames_ue4_GameActivity_nativeSetAndroidVersionInfor
 
 	FAndroidMisc::SetVersionInfo( UEAndroidVersion, UEPhoneMake, UEPhoneModel, UEOSLanguage );
 }
+
+void WaitForAndroidLoseFocusEvent()
+{
+	FAppEventManager::GetInstance()->WaitForEventInQueue(EAppEventState::APP_EVENT_STATE_WINDOW_LOST_FOCUS);
+}
