[INTEGRATE] Change 2413642 by [REDACTED_PERSON][REDACTED_HANDLE] on 2015/01/21 08:08:34

	Update default branch name string in 4.7 branch.

diff --git a/Engine/Source/Runtime/Launch/Resources/Version.h b/Engine/Source/Runtime/Launch/Resources/Version.h
--- a/Engine/Source/Runtime/Launch/Resources/Version.h
+++ b/Engine/Source/Runtime/Launch/Resources/Version.h
@@ -11,7 +11,7 @@
 #define ENGINE_IS_LICENSEE_VERSION 0
 
 #define BUILT_FROM_CHANGELIST 0
-#define BRANCH_NAME "UE4"
+#define BRANCH_NAME "++depot+UE4-Releases+4.7"
 
 #define EPIC_COMPANY_NAME  "Epic Games, Inc."
 #define EPIC_COPYRIGHT_STRING "Copyright 1998-2015 Epic Games, Inc. All Rights Reserved."
