#android Workaround for startup crash on Android OS 5 and up. (UE-8635)

diff --git a/Engine/Source/Runtime/Launch/Private/Android/AndroidJNI.cpp b/Engine/Source/Runtime/Launch/Private/Android/AndroidJNI.cpp
--- a/Engine/Source/Runtime/Launch/Private/Android/AndroidJNI.cpp
+++ b/Engine/Source/Runtime/Launch/Private/Android/AndroidJNI.cpp
@@ -549,6 +549,9 @@ extern "C" void Java_com_epicgames_ue4_GameActivity_nativeSetGlobalActivity(JNIE
 		FJavaWrapper::GoogleServicesThis = FJavaWrapper::GameActivityThis;
 		// FJavaWrapper::GoogleServicesThis = jenv->NewGlobalRef(googleServices);
 
+		bool bIsOptional = false;
+		FJavaWrapper::GameActivityClassID = FJavaWrapper::FindClass(jenv, "com/epicgames/ue4/GameActivity", bIsOptional);
+
 		// Next we check to see if the OBB file is in the APK
 		jmethodID isOBBInAPKMethod = jenv->GetStaticMethodID(FJavaWrapper::GameActivityClassID, "isOBBInAPK", "()Z");
 		GOBBinAPK = (bool)jenv->CallStaticBooleanMethod(FJavaWrapper::GameActivityClassID, isOBBInAPKMethod, nullptr);
