Don't add Audio modules if building a server on Win32, Win64, or Mac
Don't add D3D11RHI module if building a server on Win32 or Win64
Make ALAudio a Dynamically Loaded Module instead of a Private Module Dependency on Linux

diff --git a/Engine/Source/Runtime/Launch/Launch.Build.cs b/Engine/Source/Runtime/Launch/Launch.Build.cs
--- a/Engine/Source/Runtime/Launch/Launch.Build.cs
+++ b/Engine/Source/Runtime/Launch/Launch.Build.cs
@@ -47,9 +47,19 @@ public class Launch : ModuleRules
 				}
 			);
 
-			if (Target.Platform == UnrealTargetPlatform.Linux)
+			if ((Target.Platform == UnrealTargetPlatform.Win32) ||
+				(Target.Platform == UnrealTargetPlatform.Win64))
 			{
-				PrivateDependencyModuleNames.Add("ALAudio");
+				DynamicallyLoadedModuleNames.Add("D3D11RHI");
+				DynamicallyLoadedModuleNames.Add("XAudio2");
+			}
+			else if (Target.Platform == UnrealTargetPlatform.Mac)
+			{
+				DynamicallyLoadedModuleNames.Add("CoreAudio");
+			}
+			else if (Target.Platform == UnrealTargetPlatform.Linux)
+			{
+				DynamicallyLoadedModuleNames.Add("ALAudio");
 			}
 
 			PrivateIncludePathModuleNames.AddRange(
@@ -143,18 +153,6 @@ public class Launch : ModuleRules
 			}
 		}
 
-		if ((Target.Platform == UnrealTargetPlatform.Win32) ||
-						(Target.Platform == UnrealTargetPlatform.Win64))
-		{
-			DynamicallyLoadedModuleNames.Add("D3D11RHI");
-			DynamicallyLoadedModuleNames.Add("XAudio2");
-		}
-
-		if (Target.Platform == UnrealTargetPlatform.Mac)
-		{
-			DynamicallyLoadedModuleNames.Add("CoreAudio");
-		}
-
 		if (Target.Platform == UnrealTargetPlatform.IOS)
 		{
 			PrivateDependencyModuleNames.Add("OpenGLDrv");
