Fix for force feedback issue on Android
#ue4
#android

diff --git a/Engine/Build/Android/Java/src/com/epicgames/ue4/GameActivity.java b/Engine/Build/Android/Java/src/com/epicgames/ue4/GameActivity.java
--- a/Engine/Build/Android/Java/src/com/epicgames/ue4/GameActivity.java
+++ b/Engine/Build/Android/Java/src/com/epicgames/ue4/GameActivity.java
@@ -521,14 +521,32 @@ public class GameActivity extends NativeActivity
 		}
 	}
 
-	public void AndroidThunkJava_Vibrate(long Duration)
+	private class VibrateRunnable implements Runnable {
+		private int duration;
+		private Vibrator vibrator;
+
+		VibrateRunnable(final int Duration, final Vibrator vibrator)
+		{
+			this.duration = Duration;
+			this.vibrator = vibrator;
+		}
+		public void run ()
+		{
+			if (duration < 1)
+			{
+				vibrator.cancel();
+			} else {
+				vibrator.vibrate(duration);
+			}
+		}
+	}
+
+	public void AndroidThunkJava_Vibrate(int Duration)
 	{
 		Vibrator vibrator = (Vibrator)getSystemService(VIBRATOR_SERVICE);
-		if (Duration < 1)
+		if (vibrator != null)
 		{
-			vibrator.cancel();
-		} else {
-			vibrator.vibrate(Duration);
+			_activity.runOnUiThread(new VibrateRunnable(Duration, vibrator));
 		}
 	}
 
diff --git a/Engine/Source/Programs/UnrealBuildTool/Android/UEDeployAndroid.cs b/Engine/Source/Programs/UnrealBuildTool/Android/UEDeployAndroid.cs
--- a/Engine/Source/Programs/UnrealBuildTool/Android/UEDeployAndroid.cs
+++ b/Engine/Source/Programs/UnrealBuildTool/Android/UEDeployAndroid.cs
@@ -730,6 +730,7 @@ namespace UnrealBuildTool.Android
 			Text.AppendLine("\t<uses-permission android:name=\"android.permission.ACCESS_WIFI_STATE\"/>");
 			Text.AppendLine("\t<uses-permission android:name=\"android.permission.MODIFY_AUDIO_SETTINGS\"/>");
 			Text.AppendLine("\t<uses-permission android:name=\"android.permission.GET_ACCOUNTS\"/>");
+            Text.AppendLine("\t<uses-permission android:name=\"android.permission.VIBRATE\"/>");
 			Text.AppendLine("\t<uses-permission android:name=\"com.android.vending.BILLING\"/>");
 //			Text.AppendLine("\t<uses-permission android:name=\"android.permission.DISABLE_KEYGUARD\"/>");
 			if (bPackageForGearVR)
diff --git a/Engine/Source/Runtime/Core/Private/Android/AndroidInputInterface.cpp b/Engine/Source/Runtime/Core/Private/Android/AndroidInputInterface.cpp
--- a/Engine/Source/Runtime/Core/Private/Android/AndroidInputInterface.cpp
+++ b/Engine/Source/Runtime/Core/Private/Android/AndroidInputInterface.cpp
@@ -123,7 +123,7 @@ void FAndroidInputInterface::SetForceFeedbackChannelValues(int32 ControllerId, c
 	UpdateVibeMotors(NewControllerData[ControllerId]);
 }
 
-extern void AndroidThunkCpp_Vibrate(int64_t Duration);
+extern void AndroidThunkCpp_Vibrate(int32 Duration);
 
 void FAndroidInputInterface::UpdateVibeMotors(FAndroidControllerData &State)
 {
diff --git a/Engine/Source/Runtime/Launch/Private/Android/AndroidJNI.cpp b/Engine/Source/Runtime/Launch/Private/Android/AndroidJNI.cpp
--- a/Engine/Source/Runtime/Launch/Private/Android/AndroidJNI.cpp
+++ b/Engine/Source/Runtime/Launch/Private/Android/AndroidJNI.cpp
@@ -47,7 +47,7 @@ void FJavaWrapper::FindClassesAndMethods(JNIEnv* Env)
 	AndroidThunkJava_Minimize = FindMethod(Env, GameActivityClassID, "AndroidThunkJava_Minimize", "()V", bIsOptional);
 	AndroidThunkJava_ForceQuit = FindMethod(Env, GameActivityClassID, "AndroidThunkJava_ForceQuit", "()V", bIsOptional);
 	AndroidThunkJava_GetFontDirectory = FindStaticMethod(Env, GameActivityClassID, "AndroidThunkJava_GetFontDirectory", "()Ljava/lang/String;", bIsOptional);
-	AndroidThunkJava_Vibrate = FindMethod(Env, GameActivityClassID, "AndroidThunkJava_Vibrate", "(J)V", bIsOptional);
+	AndroidThunkJava_Vibrate = FindMethod(Env, GameActivityClassID, "AndroidThunkJava_Vibrate", "(I)V", bIsOptional);
 	AndroidThunkJava_IsMusicActive = FindMethod(Env, GameActivityClassID, "AndroidThunkJava_IsMusicActive", "()Z", bIsOptional);
 	AndroidThunkJava_KeepScreenOn = FindMethod(Env, GameActivityClassID, "AndroidThunkJava_KeepScreenOn", "(Z)V", bIsOptional);
 	AndroidThunkJava_InitHMDs = FindMethod(Env, GameActivityClassID, "AndroidThunkJava_InitHMDs", "()V", bIsOptional);
@@ -201,7 +201,7 @@ void AndroidThunkCpp_KeepScreenOn(bool Enable)
 	}
 }
 
-void AndroidThunkCpp_Vibrate(int64_t Duration)
+void AndroidThunkCpp_Vibrate(int32 Duration)
 {
 	if (JNIEnv* Env = FAndroidApplication::GetJavaEnv())
 	{
