[INTEGRATE] Change 2436200 by [REDACTED_PERSON][REDACTED_HANDLE] on 2015/02/06 17:38:50

	Temporarily disabled our crash reporting on Mac on build machines to help track the crash on generating DDC for StarterContent

diff --git a/Engine/Source/Runtime/Launch/Private/Mac/LaunchMac.cpp b/Engine/Source/Runtime/Launch/Private/Mac/LaunchMac.cpp
--- a/Engine/Source/Runtime/Launch/Private/Mac/LaunchMac.cpp
+++ b/Engine/Source/Runtime/Launch/Private/Mac/LaunchMac.cpp
@@ -184,7 +184,7 @@ void EngineCrashHandler(const FGenericCrashContext& GenericContext)
 #if UE_BUILD_DEBUG
 	if( true && !GAlwaysReportCrash )
 #else
-	if( FPlatformMisc::IsDebuggerPresent() && !GAlwaysReportCrash )
+	if( bIsBuildMachine || ( FPlatformMisc::IsDebuggerPresent() && !GAlwaysReportCrash ) )
 #endif
 	{
 		// Don't use exception handling when a debugger is attached to exactly trap the crash. This does NOT check
