[INTEGRATE] Change 2443483 by [REDACTED_PERSON][REDACTED_HANDLE] on 2015/02/12 10:46:25

	Fix for building Android/iOS monolithic builds.

diff --git a/Engine/Source/Runtime/Core/Private/Android/AndroidMisc.cpp b/Engine/Source/Runtime/Core/Private/Android/AndroidMisc.cpp
--- a/Engine/Source/Runtime/Core/Private/Android/AndroidMisc.cpp
+++ b/Engine/Source/Runtime/Core/Private/Android/AndroidMisc.cpp
@@ -12,7 +12,7 @@
 #include "MallocCrash.h"
 #include "AndroidJavaMessageBox.h"
 
-DEFINE_LOG_CATEGORY_STATIC(LogEngine, Log, All);
+ENGINE_API DECLARE_LOG_CATEGORY_EXTERN(LogEngine, Log, All);
 
 void* FAndroidMisc::NativeWindow = NULL;
 
diff --git a/Engine/Source/Runtime/Core/Private/IOS/IOSAppDelegate.cpp b/Engine/Source/Runtime/Core/Private/IOS/IOSAppDelegate.cpp
--- a/Engine/Source/Runtime/Core/Private/IOS/IOSAppDelegate.cpp
+++ b/Engine/Source/Runtime/Core/Private/IOS/IOSAppDelegate.cpp
@@ -15,7 +15,7 @@
 #define GAME_THREAD_STACK_SIZE 1024 * 1024 
 
 DEFINE_LOG_CATEGORY(LogIOSAudioSession);
-DEFINE_LOG_CATEGORY_STATIC(LogEngine, Log, All);
+ENGINE_API DECLARE_LOG_CATEGORY_EXTERN(LogEngine, Log, All);
 
 extern bool GShowSplashScreen;
 
diff --git a/Engine/Source/Runtime/Engine/Public/UnrealEngine.h b/Engine/Source/Runtime/Engine/Public/UnrealEngine.h
--- a/Engine/Source/Runtime/Engine/Public/UnrealEngine.h
+++ b/Engine/Source/Runtime/Engine/Public/UnrealEngine.h
@@ -8,7 +8,7 @@
 
 #include "Engine/Engine.h"
 
-DECLARE_LOG_CATEGORY_EXTERN(LogEngine, Log, All);
+ENGINE_API DECLARE_LOG_CATEGORY_EXTERN(LogEngine, Log, All);
 
 //
 //	FLocalPlayerIterator - Iterates over local players in the game.
diff --git a/Engine/Source/Runtime/Launch/Private/Android/AndroidJNI.cpp b/Engine/Source/Runtime/Launch/Private/Android/AndroidJNI.cpp
--- a/Engine/Source/Runtime/Launch/Private/Android/AndroidJNI.cpp
+++ b/Engine/Source/Runtime/Launch/Private/Android/AndroidJNI.cpp
@@ -168,7 +168,7 @@ jmethodID FJavaWrapper::AndroidThunkJava_IapQueryInAppPurchases;
 jmethodID FJavaWrapper::AndroidThunkJava_IapBeginPurchase;
 jmethodID FJavaWrapper::AndroidThunkJava_IapIsAllowedToMakePurchases;
 
-DEFINE_LOG_CATEGORY_STATIC(LogEngine, Log, All);
+ENGINE_API DECLARE_LOG_CATEGORY_EXTERN(LogEngine, Log, All);
 
 //Game-specific crash reporter
 void EngineCrashHandler(const FGenericCrashContext& GenericContext)
