[INTEGRATE] Change 2443518 by [REDACTED_PERSON][REDACTED_HANDLE] on 2015/02/12 11:07:47

	Fix Android/iOS compile errors.

diff --git a/Engine/Source/Runtime/Core/Private/Android/AndroidMisc.cpp b/Engine/Source/Runtime/Core/Private/Android/AndroidMisc.cpp
--- a/Engine/Source/Runtime/Core/Private/Android/AndroidMisc.cpp
+++ b/Engine/Source/Runtime/Core/Private/Android/AndroidMisc.cpp
@@ -12,7 +12,7 @@
 #include "MallocCrash.h"
 #include "AndroidJavaMessageBox.h"
 
-ENGINE_API DECLARE_LOG_CATEGORY_EXTERN(LogEngine, Log, All);
+DECLARE_LOG_CATEGORY_EXTERN(LogEngine, Log, All);
 
 void* FAndroidMisc::NativeWindow = NULL;
 
diff --git a/Engine/Source/Runtime/Core/Private/IOS/IOSAppDelegate.cpp b/Engine/Source/Runtime/Core/Private/IOS/IOSAppDelegate.cpp
--- a/Engine/Source/Runtime/Core/Private/IOS/IOSAppDelegate.cpp
+++ b/Engine/Source/Runtime/Core/Private/IOS/IOSAppDelegate.cpp
@@ -15,7 +15,7 @@
 #define GAME_THREAD_STACK_SIZE 1024 * 1024 
 
 DEFINE_LOG_CATEGORY(LogIOSAudioSession);
-ENGINE_API DECLARE_LOG_CATEGORY_EXTERN(LogEngine, Log, All);
+DECLARE_LOG_CATEGORY_EXTERN(LogEngine, Log, All);
 
 extern bool GShowSplashScreen;
 
diff --git a/Engine/Source/Runtime/Launch/Private/Android/AndroidJNI.cpp b/Engine/Source/Runtime/Launch/Private/Android/AndroidJNI.cpp
--- a/Engine/Source/Runtime/Launch/Private/Android/AndroidJNI.cpp
+++ b/Engine/Source/Runtime/Launch/Private/Android/AndroidJNI.cpp
@@ -168,8 +168,6 @@ jmethodID FJavaWrapper::AndroidThunkJava_IapQueryInAppPurchases;
 jmethodID FJavaWrapper::AndroidThunkJava_IapBeginPurchase;
 jmethodID FJavaWrapper::AndroidThunkJava_IapIsAllowedToMakePurchases;
 
-ENGINE_API DECLARE_LOG_CATEGORY_EXTERN(LogEngine, Log, All);
-
 //Game-specific crash reporter
 void EngineCrashHandler(const FGenericCrashContext& GenericContext)
 {
