Load HTTP module speculatively.

- Fixes UT servers crashing due to OSSMcp trying to late-cancel pending HTTP requests, when the module has already been unloaded.
- Better solution is being discussed within Core.

diff --git a/Engine/Source/Runtime/Launch/Private/LaunchEngineLoop.cpp b/Engine/Source/Runtime/Launch/Private/LaunchEngineLoop.cpp
--- a/Engine/Source/Runtime/Launch/Private/LaunchEngineLoop.cpp
+++ b/Engine/Source/Runtime/Launch/Private/LaunchEngineLoop.cpp
@@ -2666,6 +2666,12 @@ bool FEngineLoop::AppInit( )
 	// Must be after FThreadStats::StartThread();
 	// Must be before Render/RHI subsystem D3DCreate()
 	// For platform services that need D3D hooks like Steam
+	// --
+	// Why load HTTP?
+	// Because most, if not all online subsystems will load HTTP themselves. This can cause problems though, as HTTP will be loaded *after* OSS, 
+	// and if OSS holds on to resources allocated by it, this will cause crash (modules are being unloaded in LIFO order with no dependency tracking).
+	// Loading HTTP before OSS works around this problem by making ModuleManager unload HTTP after OSS, at the cost of extra module for the few OSS (like Null) that don't use it.
+	FModuleManager::Get().LoadModule(TEXT("HTTP"));
 	FModuleManager::Get().LoadModule(TEXT("OnlineSubsystem"));
 	FModuleManager::Get().LoadModule(TEXT("OnlineSubsystemUtils"));
 #endif
