Added GExternalFilePath to allow access to Android getExternalFilesDir path
#ue4
#android

diff --git a/Engine/Source/Runtime/Core/Private/Android/AndroidFile.cpp b/Engine/Source/Runtime/Core/Private/Android/AndroidFile.cpp
--- a/Engine/Source/Runtime/Core/Private/Android/AndroidFile.cpp
+++ b/Engine/Source/Runtime/Core/Private/Android/AndroidFile.cpp
@@ -41,6 +41,8 @@ static int32 GPackagePatchVersion = 0;
 
 // External File Path base - setup during load
 FString GFilePathBase;
+// External File Direcory Path (for application) - setup during load
+FString GExternalFilePath;
 // External font path base - setup during load
 FString GFontPathBase;
 
diff --git a/Engine/Source/Runtime/Launch/Private/Android/AndroidJNI.cpp b/Engine/Source/Runtime/Launch/Private/Android/AndroidJNI.cpp
--- a/Engine/Source/Runtime/Launch/Private/Android/AndroidJNI.cpp
+++ b/Engine/Source/Runtime/Launch/Private/Android/AndroidJNI.cpp
@@ -17,6 +17,7 @@ JavaVM* GJavaVM;
 static IVirtualKeyboardEntry *VirtualKeyboardWidget = NULL;
 
 extern FString GFilePathBase;
+extern FString GExternalFilePath;
 extern FString GFontPathBase;
 extern bool GOBBinAPK;
 
@@ -553,6 +554,20 @@ extern "C" void Java_com_epicgames_ue4_GameActivity_nativeSetGlobalActivity(JNIE
 		// Next we check to see if the OBB file is in the APK
 		jmethodID isOBBInAPKMethod = jenv->GetStaticMethodID(FJavaWrapper::GameActivityClassID, "isOBBInAPK", "()Z");
 		GOBBinAPK = (bool)jenv->CallStaticBooleanMethod(FJavaWrapper::GameActivityClassID, isOBBInAPKMethod, nullptr);
+
+		// Cache path to external files directory
+		jclass ContextClass = jenv->FindClass("android/content/Context");
+		jmethodID getExternalFilesDir = jenv->GetMethodID(ContextClass, "getExternalFilesDir", "(Ljava/lang/String;)Ljava/io/File;");
+		jobject externalFilesDirPath = jenv->CallObjectMethod(FJavaWrapper::GameActivityThis, getExternalFilesDir, nullptr);
+		jmethodID getFilePath = jenv->GetMethodID(jenv->FindClass("java/io/File"), "getPath", "()Ljava/lang/String;");
+		jstring externalFilesPathString = (jstring)jenv->CallObjectMethod(externalFilesDirPath, getFilePath, nullptr);
+		const char *nativeExternalFilesPathString = jenv->GetStringUTFChars(externalFilesPathString, 0);
+		// Copy that somewhere safe 
+		GExternalFilePath = FString(nativeExternalFilesPathString);
+
+		// then release...
+		jenv->ReleaseStringUTFChars(externalFilesPathString, nativeExternalFilesPathString);
+		FPlatformMisc::LowLevelOutputDebugStringf(TEXT("ExternalFilePath found as '%s'\n"), *GExternalFilePath);
 	}
 }
 
