Read Android commandline from proper file path
#ue4
#android

diff --git a/Engine/Source/Runtime/Launch/Private/Android/LaunchAndroid.cpp b/Engine/Source/Runtime/Launch/Private/Android/LaunchAndroid.cpp
--- a/Engine/Source/Runtime/Launch/Private/Android/LaunchAndroid.cpp
+++ b/Engine/Source/Runtime/Launch/Private/Android/LaunchAndroid.cpp
@@ -175,7 +175,7 @@ static void InitCommandLine()
 	FCommandLine::Set(TEXT(""));
 
 	// read in the command line text file from the sdcard if it exists
-	FString CommandLineFilePath = GFilePathBase + FString("/") + (!FApp::IsGameNameEmpty() ? FApp::GetGameName() : FPlatformProcess::ExecutableName()) + FString("/UE4CommandLine.txt");
+	FString CommandLineFilePath = GFilePathBase + FString("/UE4Game/") + (!FApp::IsGameNameEmpty() ? FApp::GetGameName() : FPlatformProcess::ExecutableName()) + FString("/UE4CommandLine.txt");
 	FILE* CommandLineFile = fopen(TCHAR_TO_UTF8(*CommandLineFilePath), "r");
 	if(CommandLineFile == NULL)
 	{
