Made the game only connect to streaming file server if you specify -streaming.  Instead of randomly choosing.

diff --git a/Engine/Source/Runtime/Launch/Private/LaunchEngineLoop.cpp b/Engine/Source/Runtime/Launch/Private/LaunchEngineLoop.cpp
--- a/Engine/Source/Runtime/Launch/Private/LaunchEngineLoop.cpp
+++ b/Engine/Source/Runtime/Launch/Private/LaunchEngineLoop.cpp
@@ -358,15 +358,23 @@ void LaunchFixGameNameCase()
 #endif	//PLATFORM_DESKTOP
 }
 
-static IPlatformFile* ConditionallyCreateFileWrapper(const TCHAR* Name, IPlatformFile* CurrentPlatformFile, const TCHAR* CommandLine, bool* OutFailedToInitialize = NULL)
+static IPlatformFile* ConditionallyCreateFileWrapper(const TCHAR* Name, IPlatformFile* CurrentPlatformFile, const TCHAR* CommandLine, bool* OutFailedToInitialize = NULL, bool* bOutShouldBeUsed = NULL )
 {
 	if (OutFailedToInitialize)
 	{
 		*OutFailedToInitialize = false;
 	}
+	if ( bOutShouldBeUsed )
+	{
+		*bOutShouldBeUsed = false;
+	}
 	IPlatformFile* WrapperFile = FPlatformFileManager::Get().GetPlatformFile(Name);
 	if (WrapperFile != NULL && WrapperFile->ShouldBeUsed(CurrentPlatformFile, CommandLine))
 	{
+		if ( bOutShouldBeUsed )
+		{
+			*bOutShouldBeUsed = true;
+		}
 		if (WrapperFile->Initialize(CurrentPlatformFile, CommandLine) == false)
 		{
 			if (OutFailedToInitialize)
@@ -427,15 +435,17 @@ bool LaunchCheckForFileOverride(const TCHAR* CmdLine, bool& OutFileOverrideFound
 	bool bNetworkFailedToInitialize = false;
 	do
 	{
-		IPlatformFile* NetworkPlatformFile = ConditionallyCreateFileWrapper(TEXT("StreamingFile"), CurrentPlatformFile, CmdLine, &bNetworkFailedToInitialize);
+		bool bShouldUseStreamingFile = false;
+		IPlatformFile* NetworkPlatformFile = ConditionallyCreateFileWrapper(TEXT("StreamingFile"), CurrentPlatformFile, CmdLine, &bNetworkFailedToInitialize, &bShouldUseStreamingFile);
 		if (NetworkPlatformFile)
 		{
 			CurrentPlatformFile = NetworkPlatformFile;
 			FPlatformFileManager::Get().SetPlatformFile(*CurrentPlatformFile);
 		}
 
+		// if streaming network platform file was tried this loop don't try this one
 		// Network file wrapper (only create if the streaming wrapper hasn't been created)
-		if (!NetworkPlatformFile)
+		if ( !bShouldUseStreamingFile && !NetworkPlatformFile)
 		{
 			NetworkPlatformFile = ConditionallyCreateFileWrapper(TEXT("NetworkFile"), CurrentPlatformFile, CmdLine, &bNetworkFailedToInitialize);
 			if (NetworkPlatformFile)
@@ -444,6 +454,7 @@ bool LaunchCheckForFileOverride(const TCHAR* CmdLine, bool& OutFileOverrideFound
 				FPlatformFileManager::Get().SetPlatformFile(*CurrentPlatformFile);
 			}
 		}
+		
 
 		if (bNetworkFailedToInitialize)
 		{
