UMG - Introducing a workaround to issues loading UMG widgets on dedicated servers.  In uncooked builds the UMG module needs to always be loaded so that widgets can be cooked, even if they wont be loaded in the cooked build.

diff --git a/Engine/Source/Runtime/Launch/Private/LaunchEngineLoop.cpp b/Engine/Source/Runtime/Launch/Private/LaunchEngineLoop.cpp
--- a/Engine/Source/Runtime/Launch/Private/LaunchEngineLoop.cpp
+++ b/Engine/Source/Runtime/Launch/Private/LaunchEngineLoop.cpp
@@ -1388,15 +1388,15 @@ int32 FEngineLoop::PreInit( const TCHAR* CmdLine )
 		GetMoviePlayer()->Initialize();
 		GetMoviePlayer()->PlayMovie();
 
-        // do any post appInit processing, before the render thread is started.
-        FPlatformMisc::PlatformPostInit(!GetMoviePlayer()->IsMovieCurrentlyPlaying());
-    }
-    else
-#endif
-    {
-        // do any post appInit processing, before the render thread is started.
-        FPlatformMisc::PlatformPostInit(true);
-    }
+		// do any post appInit processing, before the render thread is started.
+		FPlatformMisc::PlatformPostInit(!GetMoviePlayer()->IsMovieCurrentlyPlaying());
+	}
+	else
+#endif
+	{
+		// do any post appInit processing, before the render thread is started.
+		FPlatformMisc::PlatformPostInit(true);
+	}
 	SlowTask.EnterProgressFrame(5);
 
 	if (GUseThreadedRendering)
@@ -1808,14 +1808,29 @@ bool FEngineLoop::LoadStartupCoreModules()
 		FModuleManager::LoadModuleChecked<IEditorStyleModule>("EditorStyle");
 #endif //WITH_EDITOR
 
-	// Load all Development modules
-	SlowTask.EnterProgressFrame(30);
-	if (!IsRunningDedicatedServer())
+	// Load UI modules
+	SlowTask.EnterProgressFrame(10);
+	if ( !IsRunningDedicatedServer() )
 	{
 		FModuleManager::Get().LoadModule("Slate");
+	}
+
+#if WITH_EDITOR
+	// In dedicated server builds with the editor, we need to load UMG/UMGEditor for compiling blueprints.
+	// UMG must be loaded for runtime and cooking.
+	FModuleManager::Get().LoadModule("UMG");
+#else
+	if ( !IsRunningDedicatedServer() )
+	{
 		// UMG must be loaded for runtime and cooking.
 		FModuleManager::Get().LoadModule("UMG");
+	}
+#endif //WITH_EDITOR
 
+	// Load all Development modules
+	SlowTask.EnterProgressFrame(20);
+	if (!IsRunningDedicatedServer())
+	{
 #if WITH_UNREAL_DEVELOPER_TOOLS
 		FModuleManager::Get().LoadModule("MessageLog");
 		FModuleManager::Get().LoadModule("CollisionAnalyzer");
