Add support for reading UE4CommandLine.txt on Windows and Mac, as exists for packaged games on most other platforms. Use it to specify default arguments for the editor in Rocket builds, which allows removing the UE_ROCKET define.

diff --git a/Engine/Source/Programs/UnrealBuildTool/Configuration/UEBuildTarget.cs b/Engine/Source/Programs/UnrealBuildTool/Configuration/UEBuildTarget.cs
--- a/Engine/Source/Programs/UnrealBuildTool/Configuration/UEBuildTarget.cs
+++ b/Engine/Source/Programs/UnrealBuildTool/Configuration/UEBuildTarget.cs
@@ -2622,16 +2622,6 @@ namespace UnrealBuildTool
 
 			string OutputAppName = GetAppName();
 
-            // Mark it as a Rocket build
-            if (UnrealBuildTool.BuildingRocket() || UnrealBuildTool.RunningRocket())
-            {
-                GlobalCompileEnvironment.Config.Definitions.Add("UE_ROCKET=1");
-            }
-            else
-            {
-                GlobalCompileEnvironment.Config.Definitions.Add("UE_ROCKET=0");
-            }
-
 			// Installed Engine intermediates go to the project's intermediate folder. Installed Engine never writes to the engine intermediate folder. (Those files are immutable)
 			// Also, when compiling in monolithic, all intermediates go to the project's folder.  This is because a project can change definitions that affects all engine translation
 			// units too, so they can't be shared between different targets.  They are effectively project-specific engine intermediates.
diff --git a/Engine/Source/Runtime/Core/Private/Misc/RocketSupport.cpp b/Engine/Source/Runtime/Core/Private/Misc/RocketSupport.cpp
--- a/Engine/Source/Runtime/Core/Private/Misc/RocketSupport.cpp
+++ b/Engine/Source/Runtime/Core/Private/Misc/RocketSupport.cpp
@@ -10,9 +10,6 @@
 
 bool FRocketSupport::IsRocket( const TCHAR* CmdLine )
 {
-#if UE_ROCKET
-	return true;
-#else
 	static int32 RocketState = -1;
 
 	if (RocketState == -1)
@@ -33,5 +30,4 @@ bool FRocketSupport::IsRocket( const TCHAR* CmdLine )
 	}
 
 	return (RocketState == 1) ? true : false;
-#endif
 }
diff --git a/Engine/Source/Runtime/Core/Public/Misc/Build.h b/Engine/Source/Runtime/Core/Public/Misc/Build.h
--- a/Engine/Source/Runtime/Core/Public/Misc/Build.h
+++ b/Engine/Source/Runtime/Core/Public/Misc/Build.h
@@ -23,10 +23,6 @@
 #ifndef UE_BUILD_SHIPPING_EDITOR
 	#define UE_BUILD_SHIPPING_EDITOR	(#)
 #endif
-#ifndef UE_ROCKET
-	#define UE_ROCKET					0
-	#error UBT should always define UE_ROCKET to be 0 or 1
-#endif
 #ifndef UE_GAME
 	#define UE_GAME						0
 #endif
diff --git a/Engine/Source/Runtime/Launch/Private/Mac/LaunchMac.cpp b/Engine/Source/Runtime/Launch/Private/Mac/LaunchMac.cpp
--- a/Engine/Source/Runtime/Launch/Private/Mac/LaunchMac.cpp
+++ b/Engine/Source/Runtime/Launch/Private/Mac/LaunchMac.cpp
@@ -274,6 +274,35 @@ INT32_MAIN_INT32_ARGC_TCHAR_ARGV()
 		}
 		GSavedCommandLine += Argument;
 	}
+	
+	FString CommandLineFile = FPaths::RootDir() / TEXT("UE4CommandLine.txt");
+	FPaths::NormalizeFilename(CommandLineFile);
+
+	int32 FileHandle = open(TCHAR_TO_UTF8(*CommandLineFile), O_RDONLY);
+	if (FileHandle != -1)
+	{
+		int32 FileLength = lseek(FileHandle, 0, SEEK_END);
+		if(FileLength < 0 || lseek(FileHandle, 0, SEEK_SET) != 0)
+		{
+			FPlatformMisc::LowLevelOutputDebugString(TEXT("WARNING: Failed to seek in UE4CommandLine.txt"));
+		}
+		else
+		{
+			TArray<uint8> Buffer;
+			Buffer.AddZeroed(FileLength + 1);
+
+			int ReadSize = read(FileHandle, Buffer.GetData(), FileLength);
+			if(ReadSize == FileLength)
+			{
+				GSavedCommandLine += TEXT(" ") + FString(UTF8_TO_TCHAR((const ANSICHAR*)Buffer.GetData())).Replace(TEXT("\r"), TEXT(" ")).Replace(TEXT("\n"), TEXT(" ")).Trim();
+			}
+			else
+			{
+				FPlatformMisc::LowLevelOutputDebugString(TEXT("WARNING: Failed to read UE4CommandLine.txt"));
+			}
+		}
+		close(FileHandle);
+	}
 
 	SCOPED_AUTORELEASE_POOL;
 	[NSApplication sharedApplication];
diff --git a/Engine/Source/Runtime/Launch/Private/Windows/LaunchWindows.cpp b/Engine/Source/Runtime/Launch/Private/Windows/LaunchWindows.cpp
--- a/Engine/Source/Runtime/Launch/Private/Windows/LaunchWindows.cpp
+++ b/Engine/Source/Runtime/Launch/Private/Windows/LaunchWindows.cpp
@@ -139,6 +139,53 @@ LAUNCH_API int32 GuardedMainWrapper( const TCHAR* CmdLine, HINSTANCE hInInstance
 	return ErrorLevel;
 }
 
+static const TCHAR* GetCompleteCommandLine()
+{
+	static FString CommandLine;
+	if(CommandLine.Len() == 0)
+	{
+		// Default to the native command line
+		CommandLine = ::GetCommandLine();
+
+		// Get the path to the arguments file
+		FString ArgsFileName = FPaths::RootDir() / TEXT("UE4CommandLine.txt");
+
+		// Try to read the default arguments
+		HANDLE FileHandle = CreateFile(*ArgsFileName, GENERIC_READ, FILE_SHARE_READ | FILE_SHARE_DELETE, NULL, OPEN_EXISTING, FILE_ATTRIBUTE_NORMAL, NULL);
+		if(FileHandle == INVALID_HANDLE_VALUE)
+		{
+			if(GetLastError() != ERROR_FILE_NOT_FOUND)
+			{
+				FPlatformMisc::LowLevelOutputDebugString(*FString::Printf(TEXT("WARNING: Failed to open UE4CommandLine.txt - %08X\n"), GetLastError()));
+			}
+		}
+		else
+		{
+			::DWORD FileSize = ::GetFileSize(FileHandle, NULL);
+			if(FileSize == (::DWORD)0xffffffff)
+			{
+				FPlatformMisc::LowLevelOutputDebugString(*FString::Printf(TEXT("WARNING: Failed to get file size for UE4CommandLine.txt (%08x)\n"), GetLastError()));
+			}
+			else
+			{
+				TArray<BYTE> Buffer;
+				Buffer.AddZeroed(FileSize + 1);
+
+				::DWORD FileSizeRead = 0;
+				if (!ReadFile(FileHandle, Buffer.GetData(), FileSize, &FileSizeRead, NULL) || FileSizeRead < FileSize)
+				{
+					FPlatformMisc::LowLevelOutputDebugString(*FString::Printf(TEXT("WARNING: Failed to read UE4CommandLine.txt file (%08x)\n"), GetLastError()));
+				}
+				else
+				{
+					CommandLine += TEXT(" ") + FString(UTF8_TO_TCHAR((const ANSICHAR*)Buffer.GetData())).Replace(TEXT("\r"), TEXT(" ")).Replace(TEXT("\n"), TEXT(" ")).Trim();
+				}
+			}
+			CloseHandle(FileHandle);
+		}
+	}
+	return *CommandLine;
+}
 
 int32 WINAPI WinMain( HINSTANCE hInInstance, HINSTANCE hPrevInstance, char*, int32 nCmdShow )
 {
@@ -147,8 +194,8 @@ int32 WINAPI WinMain( HINSTANCE hInInstance, HINSTANCE hPrevInstance, char*, int
 
 	int32 ErrorLevel			= 0;
 	hInstance				= hInInstance;
-	const TCHAR* CmdLine	= GetCommandLine();
-	
+	const TCHAR* CmdLine = GetCompleteCommandLine();
+
 #if !(UE_BUILD_SHIPPING && WITH_EDITOR)
 	// Named mutex we use to figure out whether we are the first instance of the game running. This is needed to e.g.
 	// make sure there is no contention when trying to save the shader cache.
