PR #619 - Incremental compilation on editor startup (Contributed by [REDACTED_PERSON])

Add a "Force compile at startup" option to the Loading/Savings section of the editor preferences, which will cause the editor to always run UBT to build the project at startup. Particularly of interest to small teams that do not submit build products to VCS.

diff --git a/Engine/Source/Editor/UnrealEd/Classes/Settings/EditorLoadingSavingSettings.h b/Engine/Source/Editor/UnrealEd/Classes/Settings/EditorLoadingSavingSettings.h
--- a/Engine/Source/Editor/UnrealEd/Classes/Settings/EditorLoadingSavingSettings.h
+++ b/Engine/Source/Editor/UnrealEd/Classes/Settings/EditorLoadingSavingSettings.h
@@ -31,6 +31,10 @@ public:
 	UPROPERTY(EditAnywhere, config, Category=Startup)
 	TEnumAsByte<ELoadLevelAtStartup::Type> LoadLevelAtStartup;
 
+	/** Force project compilation at startup */
+	UPROPERTY(EditAnywhere, config, Category=Startup)
+	uint32 bForceCompilationAtStartup:1;
+
 	/** Whether to restore previously open assets at startup */
 	UPROPERTY(EditAnywhere, config, Category=Startup)
 	uint32 bRestoreOpenAssetTabsOnRestart:1;
diff --git a/Engine/Source/Runtime/Launch/Private/LaunchEngineLoop.cpp b/Engine/Source/Runtime/Launch/Private/LaunchEngineLoop.cpp
--- a/Engine/Source/Runtime/Launch/Private/LaunchEngineLoop.cpp
+++ b/Engine/Source/Runtime/Launch/Private/LaunchEngineLoop.cpp
@@ -2558,46 +2558,68 @@ bool FEngineLoop::AppInit( )
 #if UE_EDITOR
 	if(!GIsBuildMachine && FPaths::IsProjectFilePathSet() && IPluginManager::Get().AreRequiredPluginsAvailable())
 	{
-		TArray<FString> IncompatibleFiles;
-		if((IProjectManager::Get().CheckModuleCompatibility(IncompatibleFiles) & IPluginManager::Get().CheckModuleCompatibility(IncompatibleFiles)) == false)
+		const FProjectDescriptor* CurrentProject = IProjectManager::Get().GetCurrentProject();
+		if(CurrentProject != nullptr && CurrentProject->Modules.Num() > 0)
 		{
-			// Log the modules which need to be rebuilt
-			FString ModulesList = TEXT("The following modules are missing or built with a different engine version:\n\n");
-			for(int Idx = 0; Idx < IncompatibleFiles.Num(); Idx++)
-			{
-				UE_LOG(LogInit, Warning, TEXT("Incompatible or missing module: %s"), *IncompatibleFiles[Idx]);
-				ModulesList += IncompatibleFiles[Idx] + TEXT("\n");
-			}
-			ModulesList += TEXT("\nWould you like to rebuild them now?");
+			bool bNeedCompile = false;
+			GConfig->GetBool(TEXT("/Script/UnrealEd.EditorLoadingSavingSettings"), TEXT("bForceCompilationAtStartup"), bNeedCompile, GEditorUserSettingsIni);
 
-			// If we're running with -stdout, assume that we're a non interactive process and about to fail
-			if (FApp::IsUnattended() || FParse::Param(FCommandLine::Get(), TEXT("stdout")))
+			if(!bNeedCompile)
 			{
-				return false;
-			}
+				// Check if any of the project or plugin modules are out of date, and the user wants to compile them.
+				TArray<FString> IncompatibleFiles;
+				IProjectManager::Get().CheckModuleCompatibility(IncompatibleFiles);
+				IPluginManager::Get().CheckModuleCompatibility(IncompatibleFiles);
 
-			// Ask whether to compile before continuing
-			if(FPlatformMisc::MessageBoxExt(EAppMsgType::YesNo, *ModulesList, *FString::Printf(TEXT("Missing %s Modules"), FApp::GetGameName())) == EAppReturnType::No)
-			{
-				return false;
-			}
+				if (IncompatibleFiles.Num() > 0)
+				{
+					// Log the modules which need to be rebuilt
+					FString ModulesList = TEXT("The following modules are missing or built with a different engine version:\n\n");
+					for (int Idx = 0; Idx < IncompatibleFiles.Num(); Idx++)
+					{
+						UE_LOG(LogInit, Warning, TEXT("Incompatible or missing module: %s"), *IncompatibleFiles[Idx]);
+						ModulesList += IncompatibleFiles[Idx] + TEXT("\n");
+					}
+					ModulesList += TEXT("\nWould you like to rebuild them now?");
 
-			// Try to compile it
-			FFeedbackContext *Context = (FFeedbackContext*)FDesktopPlatformModule::Get()->GetNativeFeedbackContext();
-			Context->BeginSlowTask(FText::FromString(TEXT("Starting build...")), true, true);
-			bool bCompileResult = FDesktopPlatformModule::Get()->CompileGameProject(FPaths::RootDir(), FPaths::GetProjectFilePath(), Context);
-			Context->EndSlowTask();
+					// If we're running with -stdout, assume that we're a non interactive process and about to fail
+					if (FApp::IsUnattended() || FParse::Param(FCommandLine::Get(), TEXT("stdout")))
+					{
+						return false;
+					}
+
+					// Ask whether to compile before continuing
+					if (FPlatformMisc::MessageBoxExt(EAppMsgType::YesNo, *ModulesList, *FString::Printf(TEXT("Missing %s Modules"), FApp::GetGameName())) == EAppReturnType::No)
+					{
+						return false;
+					}
+
+					bNeedCompile = true;
+				}
+			}
 
-			// Check it succeeded
-			TArray<FString> StillIncompatibleFiles;
-			if(!bCompileResult || !IProjectManager::Get().CheckModuleCompatibility(StillIncompatibleFiles) || !IPluginManager::Get().CheckModuleCompatibility(StillIncompatibleFiles))
+			if(bNeedCompile)
 			{
-				for(int Idx = 0; Idx < StillIncompatibleFiles.Num(); Idx++)
+				// Try to compile it
+				FFeedbackContext *Context = (FFeedbackContext*)FDesktopPlatformModule::Get()->GetNativeFeedbackContext();
+				Context->BeginSlowTask(FText::FromString(TEXT("Starting build...")), true, true);
+				bool bCompileResult = FDesktopPlatformModule::Get()->CompileGameProject(FPaths::RootDir(), FPaths::GetProjectFilePath(), Context);
+				Context->EndSlowTask();
+
+				// Get a list of modules which are still incompatible
+				TArray<FString> StillIncompatibleFiles;
+				IProjectManager::Get().CheckModuleCompatibility(StillIncompatibleFiles);
+				IPluginManager::Get().CheckModuleCompatibility(StillIncompatibleFiles);
+
+				if(!bCompileResult || StillIncompatibleFiles.Num() > 0)
 				{
-					UE_LOG(LogInit, Warning, TEXT("Still incompatible or missing module: %s"), *StillIncompatibleFiles[Idx]);
+					for (int Idx = 0; Idx < StillIncompatibleFiles.Num(); Idx++)
+					{
+						UE_LOG(LogInit, Warning, TEXT("Still incompatible or missing module: %s"), *StillIncompatibleFiles[Idx]);
+					}
+					FPlatformMisc::MessageBoxExt(EAppMsgType::Ok, *FString::Printf(TEXT("%s could not be compiled. Try rebuilding from source manually."), FApp::GetGameName()), TEXT("Error"));
+					return false;
 				}
-				FPlatformMisc::MessageBoxExt(EAppMsgType::Ok, *FString::Printf(TEXT("%s could not be compiled. Try rebuilding from source manually."), FApp::GetGameName()), TEXT("Error"));
-				return false;
 			}
 		}
 	}
@@ -2872,4 +2894,4 @@ void FEngineLoop::PreInitHMDDevice()
 #endif // #if WITH_ENGINE && !UE_SERVER
 }
 
-#undef LOCTEXT_NAMESPACE
+#undef LOCTEXT_NAMESPACE
\ No newline at end of file
