Move resources entwined with preprocessor definitions from PCLaunch.rc to PCLaunch.rc2, and include that from PCLaunch.rc so that the Visual Studio resource editor can read and write PCLaunch.rc correctly.

diff --git a/Engine/Source/Runtime/Launch/Resources/Windows/PCLaunch.rc b/Engine/Source/Runtime/Launch/Resources/Windows/PCLaunch.rc
--- a/Engine/Source/Runtime/Launch/Resources/Windows/PCLaunch.rc
+++ b/Engine/Source/Runtime/Launch/Resources/Windows/PCLaunch.rc
@@ -1,7 +1,6 @@
-// Copyright 1998-2015 Epic Games, Inc. All Rights Reserved.
-
+// Microsoft Visual C++ generated resource script.
+//
 #include "resource.h"
-#include "CoreMiscDefines.h"
 
 #define APSTUDIO_READONLY_SYMBOLS
 /////////////////////////////////////////////////////////////////////////////
@@ -9,20 +8,16 @@
 // Generated from the TEXTINCLUDE 2 resource.
 //
 #include <windows.h>
-#include "../Version.h"
 
 /////////////////////////////////////////////////////////////////////////////
 #undef APSTUDIO_READONLY_SYMBOLS
 
-
 /////////////////////////////////////////////////////////////////////////////
-// English (U.S.) resources
+// English (United States) resources
 
 #if !defined(AFX_RESOURCE_DLL) || defined(AFX_TARG_ENU)
-#ifdef _WIN32
 LANGUAGE LANG_ENGLISH, SUBLANG_ENGLISH_US
 #pragma code_page(1252)
-#endif //_WIN32
 
 /////////////////////////////////////////////////////////////////////////////
 //
@@ -38,7 +33,7 @@ BEGIN
 	PUSHBUTTON      "No",IDC_NO_B,141,97,50,14
 	PUSHBUTTON      "Yes to all",IDC_YESTOALL,83,97,50,14
 	PUSHBUTTON      "No to all",IDC_NOTOALL,199,97,50,14
-	LTEXT           "Static",IDC_MESSAGE,17,19,246,66
+    LTEXT           "Static",IDC_MESSAGE,17,19,246,66
 END
 
 IDD_YESNO2ALLCANCEL DIALOGEX 0, 0, 297, 126
@@ -46,40 +41,42 @@ STYLE DS_SETFONT | DS_MODALFRAME | DS_FIXEDSYS | DS_CENTER | WS_POPUP | WS_CAPTI
 CAPTION "Dialog"
 FONT 8, "MS Shell Dlg", 400, 0, 0x1
 BEGIN
-	PUSHBUTTON      "Yes",IDC_YES,15,98,50,14
-	PUSHBUTTON      "Yes to all",IDC_YESTOALL,70,98,50,14
-	PUSHBUTTON      "No",IDC_NO_B,124,98,50,14
-	PUSHBUTTON      "No to all",IDC_NOTOALL,177,98,50,14
-	PUSHBUTTON      "Cancel",IDC_CANCEL,232,98,50,14
-	LTEXT           "Static",IDC_MESSAGE,16,14,265,66
+    PUSHBUTTON      "Yes",IDC_YES,15,98,50,14
+    PUSHBUTTON      "Yes to all",IDC_YESTOALL,70,98,50,14
+    PUSHBUTTON      "No",IDC_NO_B,124,98,50,14
+    PUSHBUTTON      "No to all",IDC_NOTOALL,177,98,50,14
+    PUSHBUTTON      "Cancel",IDC_CANCEL,232,98,50,14
+    LTEXT           "Static",IDC_MESSAGE,16,14,265,66
 END
 
+
 /////////////////////////////////////////////////////////////////////////////
 //
 // DESIGNINFO
 //
 
 #ifdef APSTUDIO_INVOKED
-GUIDELINES DESIGNINFO 
+GUIDELINES DESIGNINFO
 BEGIN
-	IDD_YESNO2ALL, DIALOG
-	BEGIN
-		LEFTMARGIN, 7
+    IDD_YESNO2ALL, DIALOG
+    BEGIN
+        LEFTMARGIN, 7
 		RIGHTMARGIN, 268
-		TOPMARGIN, 7
+        TOPMARGIN, 7
 		BOTTOMMARGIN, 111
-	END
-
-	IDD_YESNO2ALLCANCEL, DIALOG
-	BEGIN
-		LEFTMARGIN, 7
-		RIGHTMARGIN, 290
-		TOPMARGIN, 7
-		BOTTOMMARGIN, 119
-	END
+    END
+
+    IDD_YESNO2ALLCANCEL, DIALOG
+    BEGIN
+        LEFTMARGIN, 7
+        RIGHTMARGIN, 290
+        TOPMARGIN, 7
+        BOTTOMMARGIN, 119
+    END
 END
 #endif    // APSTUDIO_INVOKED
 
+
 #ifdef APSTUDIO_INVOKED
 /////////////////////////////////////////////////////////////////////////////
 //
@@ -88,84 +85,24 @@ END
 
 1 TEXTINCLUDE 
 BEGIN
-	"resource.h\0"
+    "resource.h\0"
 END
 
 2 TEXTINCLUDE 
 BEGIN
-	"#include <windows.h>\r\n"
-	"\0"
+    "#include <windows.h>\r\n"
+    "\0"
 END
 
 3 TEXTINCLUDE 
 BEGIN
-	"\r\n"
-	"\0"
+    "#include ""PCLaunch.rc2""\0"
 END
 
 #endif    // APSTUDIO_INVOKED
 
-/////////////////////////////////////////////////////////////////////////////
-//
-// Version
-//
 
-VS_VERSION_INFO VERSIONINFO
- FILEVERSION ENGINE_MAJOR_VERSION,ENGINE_MINOR_VERSION,ENGINE_PATCH_VERSION,0
- PRODUCTVERSION ENGINE_MAJOR_VERSION,ENGINE_MINOR_VERSION,ENGINE_PATCH_VERSION,0
- FILEFLAGSMASK 0x17L
-#ifdef _DEBUG
- FILEFLAGS 0x1L
-#else
- FILEFLAGS 0x0L
-#endif
- FILEOS 0x4L
- FILETYPE 0x2L
- FILESUBTYPE 0x0L
-BEGIN
-	BLOCK "StringFileInfo"
-	BEGIN
-		BLOCK "040904b0"
-		BEGIN
-			VALUE "CompanyName", EPIC_COMPANY_NAME
-			VALUE "LegalCopyright", EPIC_COPYRIGHT_STRING
-			VALUE "ProductName", EPIC_PRODUCT_NAME
-			VALUE "ProductVersion", ENGINE_VERSION_STRING
-			VALUE "FileDescription", EPIC_PRODUCT_NAME
-			VALUE "InternalName", EPIC_PRODUCT_IDENTIFIER
-			VALUE "OriginalFilename", PREPROCESSOR_TO_STRING(ORIGINAL_FILE_NAME)
-		END
-	END
-	BLOCK "VarFileInfo"
-	BEGIN
-		VALUE "Translation", 0x409, 1200
-	END
-END
 
-/////////////////////////////////////////////////////////////////////////////
-//
-// Application Manifest
-//
-
-#if _DEBUG
-
-// Debug manifest
-#if _WIN64
-1						RT_MANIFEST				"Debug-PCLaunch-Win64.manifest"
-#else
-1						RT_MANIFEST				"Debug-PCLaunch-Win32.manifest"
-#endif
-
-#else
-
-// Release manifest
-#if _WIN64
-1						RT_MANIFEST				"PCLaunch-Win64.manifest"
-#else
-1						RT_MANIFEST				"PCLaunch-Win32.manifest"
-#endif
-
-#endif
 
 /////////////////////////////////////////////////////////////////////////////
 //
@@ -175,19 +112,17 @@ END
 // Icon with lowest ID value placed first to ensure application icon
 // remains consistent on all systems.
 IDICON_UE4Game          ICON                    "UE4.ico"
-
-
-#endif    // English (U.S.) resources
+#endif    // English (United States) resources
 /////////////////////////////////////////////////////////////////////////////
 
 
 
-//#ifndef APSTUDIO_INVOKED
+#ifndef APSTUDIO_INVOKED
 /////////////////////////////////////////////////////////////////////////////
 //
 // Generated from the TEXTINCLUDE 3 resource.
 //
-
-
+#include "PCLaunch.rc2"
 /////////////////////////////////////////////////////////////////////////////
-//#endif    // not APSTUDIO_INVOKED
+#endif    // not APSTUDIO_INVOKED
+
diff --git a/Engine/Source/Runtime/Launch/Resources/Windows/PCLaunch.rc2 b/Engine/Source/Runtime/Launch/Resources/Windows/PCLaunch.rc2
new file mode 100644
--- /dev/null
+++ b/Engine/Source/Runtime/Launch/Resources/Windows/PCLaunch.rc2
@@ -0,0 +1,70 @@
+// Copyright 1998-2015 Epic Games, Inc. All Rights Reserved.
+//
+// This file contains resources which are edited by hand, and cannot be loaded/saved by the Visual Studio resource editor.
+
+#include "resource.h"
+#include "CoreMiscDefines.h"
+#include "../Version.h"
+
+/////////////////////////////////////////////////////////////////////////////
+//
+// Version
+//
+
+VS_VERSION_INFO VERSIONINFO
+ FILEVERSION ENGINE_MAJOR_VERSION,ENGINE_MINOR_VERSION,ENGINE_PATCH_VERSION,0
+ PRODUCTVERSION ENGINE_MAJOR_VERSION,ENGINE_MINOR_VERSION,ENGINE_PATCH_VERSION,0
+ FILEFLAGSMASK 0x17L
+#ifdef _DEBUG
+ FILEFLAGS 0x1L
+#else
+ FILEFLAGS 0x0L
+#endif
+ FILEOS 0x4L
+ FILETYPE 0x2L
+ FILESUBTYPE 0x0L
+BEGIN
+	BLOCK "StringFileInfo"
+	BEGIN
+		BLOCK "040904b0"
+		BEGIN
+			VALUE "CompanyName", EPIC_COMPANY_NAME
+			VALUE "LegalCopyright", EPIC_COPYRIGHT_STRING
+			VALUE "ProductName", EPIC_PRODUCT_NAME
+			VALUE "ProductVersion", ENGINE_VERSION_STRING
+			VALUE "FileDescription", EPIC_PRODUCT_NAME
+			VALUE "InternalName", EPIC_PRODUCT_IDENTIFIER
+			VALUE "OriginalFilename", PREPROCESSOR_TO_STRING(ORIGINAL_FILE_NAME)
+		END
+	END
+	BLOCK "VarFileInfo"
+	BEGIN
+		VALUE "Translation", 0x409, 1200
+	END
+END
+
+/////////////////////////////////////////////////////////////////////////////
+//
+// Application Manifest
+//
+
+#if _DEBUG
+
+// Debug manifest
+#if _WIN64
+1						RT_MANIFEST				"Debug-PCLaunch-Win64.manifest"
+#else
+1						RT_MANIFEST				"Debug-PCLaunch-Win32.manifest"
+#endif
+
+#else
+
+// Release manifest
+#if _WIN64
+1						RT_MANIFEST				"PCLaunch-Win64.manifest"
+#else
+1						RT_MANIFEST				"PCLaunch-Win32.manifest"
+#endif
+
+#endif
+
diff --git a/Engine/Source/Runtime/Launch/Resources/Windows/resource.h b/Engine/Source/Runtime/Launch/Resources/Windows/resource.h
--- a/Engine/Source/Runtime/Launch/Resources/Windows/resource.h
+++ b/Engine/Source/Runtime/Launch/Resources/Windows/resource.h
@@ -1,5 +1,3 @@
-// Copyright 1998-2015 Epic Games, Inc. All Rights Reserved.
-
 //{{NO_DEPENDENCIES}}
 // Microsoft Visual C++ generated include file.
 // Used by PCLaunch.rc
@@ -13,10 +11,8 @@
 #define IDC_NO_B                        1003
 #define IDC_MESSAGE                     1004
 #define IDC_CANCEL                      1009
-/** ID for SHA hash file */
 #define ID_HASHFILE                     1010
 
-
 // Next default values for new objects
 // 
 #ifdef APSTUDIO_INVOKED
@@ -27,4 +23,3 @@
 #define _APS_NEXT_SYMED_VALUE           101
 #endif
 #endif
-
