Fix up some depreciated calls for Android
#ue4
#android

diff --git a/Engine/Source/Runtime/Launch/Private/Android/AndroidEventManager.cpp b/Engine/Source/Runtime/Launch/Private/Android/AndroidEventManager.cpp
--- a/Engine/Source/Runtime/Launch/Private/Android/AndroidEventManager.cpp
+++ b/Engine/Source/Runtime/Launch/Private/Android/AndroidEventManager.cpp
@@ -269,9 +269,9 @@ void FAppEventManager::PauseAudio()
 {
 	bAudioPaused = true;
 
-	if (GEngine->GetAudioDevice())
+	if (GEngine->GetMainAudioDevice())
 	{
-		GEngine->GetAudioDevice()->Suspend(false);
+		GEngine->GetMainAudioDevice()->Suspend(false);
 	}
 }
 
@@ -280,9 +280,9 @@ void FAppEventManager::ResumeAudio()
 {
 	bAudioPaused = false;
 
-	if (GEngine->GetAudioDevice())
+	if (GEngine->GetMainAudioDevice())
 	{
-		GEngine->GetAudioDevice()->Suspend(true);
+		GEngine->GetMainAudioDevice()->Suspend(true);
 	}
 }
 
diff --git a/Engine/Source/Runtime/Launch/Private/Android/LaunchAndroid.cpp b/Engine/Source/Runtime/Launch/Private/Android/LaunchAndroid.cpp
--- a/Engine/Source/Runtime/Launch/Private/Android/LaunchAndroid.cpp
+++ b/Engine/Source/Runtime/Launch/Private/Android/LaunchAndroid.cpp
@@ -281,7 +281,7 @@ int32 AndroidMain(struct android_app* state)
 	InitCommandLine();
 	FPlatformMisc::LowLevelOutputDebugStringf(TEXT("Final commandline: %s\n"), FCommandLine::Get());
 
-	EventHandlerEvent = FPlatformProcess::CreateSynchEvent();
+	EventHandlerEvent = FPlatformProcess::GetSynchEventFromPool(false);
 	FPlatformMisc::LowLevelOutputDebugString(L"Created sync event");
 	FAppEventManager::GetInstance()->SetEventHandlerEvent(EventHandlerEvent);
 
@@ -293,6 +293,16 @@ int32 AndroidMain(struct android_app* state)
 	// OBBs and APK are found.
 	IPlatformFile::GetPlatformPhysical().Initialize(nullptr, FCommandLine::Get());
 
+#if 0
+	for (int32 i = 0; i < 10; i++)
+	{
+		sleep(1);
+		FPlatformMisc::LowLevelOutputDebugStringf(TEXT("[Patch %d]"), i);
+
+	}
+	FPlatformMisc::LowLevelOutputDebugStringf(TEXT("[Patch] : Dont Patch \n"));
+#endif
+
 	// initialize the engine
 	GEngineLoop.PreInit(0, NULL, FCommandLine::Get());
 
