Removed logging of IsRocket status at startup as it's not necessary, it will already be in the commandline displayed previously.

diff --git a/Engine/Source/Runtime/Launch/Private/LaunchEngineLoop.cpp b/Engine/Source/Runtime/Launch/Private/LaunchEngineLoop.cpp
--- a/Engine/Source/Runtime/Launch/Private/LaunchEngineLoop.cpp
+++ b/Engine/Source/Runtime/Launch/Private/LaunchEngineLoop.cpp
@@ -2720,7 +2720,6 @@ bool FEngineLoop::AppInit( )
 	UE_LOG(LogInit, Log, TEXT("Command line: %s"), FCommandLine::Get() );
 	UE_LOG(LogInit, Log, TEXT("Base directory: %s"), FPlatformProcess::BaseDir() );
 	//UE_LOG(LogInit, Log, TEXT("Character set: %s"), sizeof(TCHAR)==1 ? TEXT("ANSI") : TEXT("Unicode") );
-	UE_LOG(LogInit, Log, TEXT("Rocket: %d"), FRocketSupport::IsRocket()? 1 : 0);
 
 	GPrintLogTimes = ELogTimes::None;
 
