UEPLAT-593 implement lifetime event callbacks as game thread tasks (Android)
#ue4
#android

diff --git a/Engine/Source/Runtime/Launch/Private/Android/AndroidEventManager.cpp b/Engine/Source/Runtime/Launch/Private/Android/AndroidEventManager.cpp
--- a/Engine/Source/Runtime/Launch/Private/Android/AndroidEventManager.cpp
+++ b/Engine/Source/Runtime/Launch/Private/Android/AndroidEventManager.cpp
@@ -66,7 +66,14 @@ void FAppEventManager::Tick()
 			//doing nothing here
 			break;
 		case APP_EVENT_STATE_ON_DESTROY:
-			FCoreDelegates::ApplicationWillTerminateDelegate.Broadcast();
+			if (FTaskGraphInterface::IsRunning())
+			{
+				FGraphEventRef WillTerminateTask = FFunctionGraphTask::CreateAndDispatchWhenReady([&]()
+				{
+					FCoreDelegates::ApplicationWillTerminateDelegate.Broadcast();
+				}, TStatId(), NULL, ENamedThreads::GameThread);
+				FTaskGraphInterface::Get().WaitUntilTaskCompletes(WillTerminateTask);
+			}
 			GIsRequestingExit = true; //destroy immediately. Game will shutdown.
 			break;
 		case APP_EVENT_STATE_ON_STOP:
@@ -123,16 +130,36 @@ void FAppEventManager::Tick()
 			ResumeAudio();
 
 			// broadcast events after the rendering thread has resumed
-			FCoreDelegates::ApplicationHasEnteredForegroundDelegate.Broadcast();
-			FCoreDelegates::ApplicationHasReactivatedDelegate.Broadcast();
+			if (FTaskGraphInterface::IsRunning())
+			{
+				FGraphEventRef EnterForegroundTask = FFunctionGraphTask::CreateAndDispatchWhenReady([&]()
+				{
+					FCoreDelegates::ApplicationHasEnteredForegroundDelegate.Broadcast();
+				}, TStatId(), NULL, ENamedThreads::GameThread);
+				FGraphEventRef ReactivateTask = FFunctionGraphTask::CreateAndDispatchWhenReady([&]()
+				{
+					FCoreDelegates::ApplicationHasReactivatedDelegate.Broadcast();
+				}, TStatId(), EnterForegroundTask, ENamedThreads::GameThread);
+				FTaskGraphInterface::Get().WaitUntilTaskCompletes(ReactivateTask);
+			}
 
 			bRunning = true;
 		}
 		else if (bRunning && (!bHaveWindow || !bHaveGame))
 		{
-			// broadcast events before rendering thred suspends
-			FCoreDelegates::ApplicationWillDeactivateDelegate.Broadcast();
-			FCoreDelegates::ApplicationWillEnterBackgroundDelegate.Broadcast();
+			// broadcast events before rendering thread suspends
+			if (FTaskGraphInterface::IsRunning())
+			{
+				FGraphEventRef DeactivateTask = FFunctionGraphTask::CreateAndDispatchWhenReady([&]()
+				{
+					FCoreDelegates::ApplicationWillDeactivateDelegate.Broadcast();
+				}, TStatId(), NULL, ENamedThreads::GameThread);
+				FGraphEventRef EnterBackgroundTask = FFunctionGraphTask::CreateAndDispatchWhenReady([&]()
+				{
+					FCoreDelegates::ApplicationWillEnterBackgroundDelegate.Broadcast();
+				}, TStatId(), DeactivateTask, ENamedThreads::GameThread);
+				FTaskGraphInterface::Get().WaitUntilTaskCompletes(EnterBackgroundTask);
+			}
 
 			PauseRendering();
 			PauseAudio();
