Core - Stats based memory profiler should start on all platforms (missing file)

diff --git a/Engine/Source/Runtime/Launch/Private/LaunchEngineLoop.cpp b/Engine/Source/Runtime/Launch/Private/LaunchEngineLoop.cpp
--- a/Engine/Source/Runtime/Launch/Private/LaunchEngineLoop.cpp
+++ b/Engine/Source/Runtime/Launch/Private/LaunchEngineLoop.cpp
@@ -7,6 +7,7 @@
 #include "ExceptionHandling.h"
 #include "FileManagerGeneric.h"
 #include "TaskGraphInterfaces.h"
+#include "StatsMallocProfilerProxy.h"
 #include "Runtime/Core/Public/Modules/ModuleVersion.h"
 
 #include "Projects.h"
@@ -697,6 +698,15 @@ int32 FEngineLoop::PreInit( const TCHAR* CmdLine )
 		return -1;
 	}
 
+#if	STATS
+	// Create the stats malloc profiler proxy.
+	if( FStatsMallocProfilerProxy::HasMemoryProfilerToken() )
+	{
+		// Assumes no concurrency here.
+		GMalloc = FStatsMallocProfilerProxy::Get();
+	}
+#endif // STATS
+
 	// Set GameName, based on the command line
 	if (LaunchSetGameName(CmdLine) == false)
 	{
