Fix wording for missed filehost. It's going to retry, don't say exit.

diff --git a/Engine/Source/Runtime/Launch/Private/LaunchEngineLoop.cpp b/Engine/Source/Runtime/Launch/Private/LaunchEngineLoop.cpp
--- a/Engine/Source/Runtime/Launch/Private/LaunchEngineLoop.cpp
+++ b/Engine/Source/Runtime/Launch/Private/LaunchEngineLoop.cpp
@@ -461,7 +461,7 @@ bool LaunchCheckForFileOverride(const TCHAR* CmdLine, bool& OutFileOverrideFound
 			FString HostIpString;
 			FParse::Value(CmdLine, TEXT("-FileHostIP="), HostIpString);
 #if PLATFORM_REQUIRES_FILESERVER
-			FPlatformMisc::LowLevelOutputDebugStringf(TEXT("Failed to connect to file server at %s. EXITING.\n"), *HostIpString);
+			FPlatformMisc::LowLevelOutputDebugStringf(TEXT("Failed to connect to file server at %s. RETRYING.\n"), *HostIpString);
 			uint32 Result = 2;
 #else	//PLATFORM_REQUIRES_FILESERVER
 			// note that this can't be localized because it happens before we connect to a filserver - localizing would cause ICU to try to load.... from over the file server connection!
