Core - Added CheckForPrintTimesOverride and moved it earlier

diff --git a/Engine/Source/Runtime/Launch/Private/LaunchEngineLoop.cpp b/Engine/Source/Runtime/Launch/Private/LaunchEngineLoop.cpp
--- a/Engine/Source/Runtime/Launch/Private/LaunchEngineLoop.cpp
+++ b/Engine/Source/Runtime/Launch/Private/LaunchEngineLoop.cpp
@@ -1063,6 +1063,7 @@ int32 FEngineLoop::PreInit( const TCHAR* CmdLine )
 
 #if STATS
 	FThreadStats::StartThread();
+	UE_LOG( LogStats, Log, TEXT( "Stats thread started" ) );
 #endif
 
 	if (FPlatformProcess::SupportsMultithreading())
@@ -2489,6 +2490,38 @@ void FEngineLoop::ClearPendingCleanupObjects()
 
 #endif // WITH_ENGINE
 
+static void CheckForPrintTimesOverride()
+{
+	GPrintLogTimes = ELogTimes::None;
+
+	// Determine whether to override the default setting for including timestamps in the log.
+	FString LogTimes;
+	if (GConfig->GetString( TEXT( "LogFiles" ), TEXT( "LogTimes" ), LogTimes, GEngineIni ))
+	{
+		if (LogTimes == TEXT( "SinceStart" ))
+		{
+			GPrintLogTimes = ELogTimes::SinceGStartTime;
+		}
+		// Assume this is a bool for backward compatibility
+		else if (FCString::ToBool( *LogTimes ))
+		{
+			GPrintLogTimes = ELogTimes::UTC;
+		}
+	}
+
+	if (FParse::Param( FCommandLine::Get(), TEXT( "LOGTIMES" ) ))
+	{
+		GPrintLogTimes = ELogTimes::UTC;
+	}
+	else if (FParse::Param( FCommandLine::Get(), TEXT( "NOLOGTIMES" ) ))
+	{
+		GPrintLogTimes = ELogTimes::None;
+	}
+	else if (FParse::Param( FCommandLine::Get(), TEXT( "LOGTIMESINCESTART" ) ))
+	{
+		GPrintLogTimes = ELogTimes::SinceGStartTime;
+	}
+}
 
 /* FEngineLoop static interface
  *****************************************************************************/
@@ -2570,6 +2603,8 @@ bool FEngineLoop::AppInit( )
 	// init config system
 	FConfigCacheIni::InitializeConfigSystem();
 
+	CheckForPrintTimesOverride();
+
 	// Check whether the project or any of its plugins are missing or are out of date
 #if UE_EDITOR
 	if(!GIsBuildMachine && FPaths::IsProjectFilePathSet() && IPluginManager::Get().AreRequiredPluginsAvailable())
@@ -2731,36 +2766,6 @@ bool FEngineLoop::AppInit( )
 	UE_LOG(LogInit, Log, TEXT("Base directory: %s"), FPlatformProcess::BaseDir() );
 	//UE_LOG(LogInit, Log, TEXT("Character set: %s"), sizeof(TCHAR)==1 ? TEXT("ANSI") : TEXT("Unicode") );
 
-	GPrintLogTimes = ELogTimes::None;
-
-	// Determine whether to override the default setting for including timestamps in the log.
-	FString LogTimes;
-	if (GConfig->GetString(TEXT("LogFiles"), TEXT("LogTimes"), LogTimes, GEngineIni))
-	{
-		if (LogTimes == TEXT("SinceStart"))
-		{
-			GPrintLogTimes = ELogTimes::SinceGStartTime;
-		}
-		// Assume this is a bool for backward compatibility
-		else if (FCString::ToBool(*LogTimes))
-		{
-			GPrintLogTimes = ELogTimes::UTC;
-		}
-	}
-
-	if (FParse::Param(FCommandLine::Get(), TEXT("LOGTIMES")))
-	{
-		GPrintLogTimes = ELogTimes::UTC;
-	}
-	else if(FParse::Param(FCommandLine::Get(), TEXT("NOLOGTIMES")))
-	{
-		GPrintLogTimes = ELogTimes::None;
-	}
-	else if(FParse::Param(FCommandLine::Get(), TEXT("LOGTIMESINCESTART")))
-	{
-		GPrintLogTimes = ELogTimes::SinceGStartTime;
-	}
-
 	// if a logging build, clear out old log files
 #if !NO_LOGGING && !(UE_BUILD_SHIPPING || UE_BUILD_TEST)
 	FMaintenance::DeleteOldLogs();
