Fixup commandline if project path is incorrect.

diff --git a/Engine/Source/Runtime/Launch/Private/LaunchEngineLoop.cpp b/Engine/Source/Runtime/Launch/Private/LaunchEngineLoop.cpp
--- a/Engine/Source/Runtime/Launch/Private/LaunchEngineLoop.cpp
+++ b/Engine/Source/Runtime/Launch/Private/LaunchEngineLoop.cpp
@@ -231,7 +231,7 @@ bool ParseGameProjectFromCommandLine(const TCHAR* InCmdLine, FString& OutProject
 	return false;
 }
 
-bool LaunchSetGameName(const TCHAR *InCmdLine)
+bool LaunchSetGameName(const TCHAR *InCmdLine, FString& OutGameProjectFilePathUnnormalized)
 {
 	if (GIsGameAgnosticExe)
 	{
@@ -247,6 +247,7 @@ bool LaunchSetGameName(const TCHAR *InCmdLine)
 			{
 				FApp::SetGameName(*LocalGameName);
 			}
+			OutGameProjectFilePathUnnormalized = ProjFilePath;
 			FPaths::SetProjectFilePath(ProjFilePath);
 		}
 #if UE_GAME
@@ -260,6 +261,7 @@ bool LaunchSetGameName(const TCHAR *InCmdLine)
 			if (LocalGameName != TEXT("UE4Game"))
 			{
 				ProjFilePath = FPaths::Combine(TEXT(".."), TEXT(".."), TEXT(".."), *LocalGameName, *FString(LocalGameName + TEXT(".") + FProjectDescriptor::GetExtension()));
+				OutGameProjectFilePathUnnormalized = ProjFilePath;
 				FPaths::SetProjectFilePath(ProjFilePath);
 			}
 		}
@@ -305,6 +307,7 @@ bool LaunchSetGameName(const TCHAR *InCmdLine)
 			{
 				FApp::SetGameName(*LocalGameName);
 			}
+			OutGameProjectFilePathUnnormalized = ProjFilePath;
 			FPaths::SetProjectFilePath(ProjFilePath);
 		}
 
@@ -708,8 +711,12 @@ int32 FEngineLoop::PreInit( const TCHAR* CmdLine )
 	}
 #endif // STATS
 
+	// Name of project file before normalization (as specified in command line).
+	// Used to fixup project name if necessary.
+	FString GameProjectFilePathUnnormalized;
+
 	// Set GameName, based on the command line
-	if (LaunchSetGameName(CmdLine) == false)
+	if (LaunchSetGameName(CmdLine, GameProjectFilePathUnnormalized) == false)
 	{
 		// If it failed, do not continue
 		return 1;
@@ -763,6 +770,12 @@ int32 FEngineLoop::PreInit( const TCHAR* CmdLine )
 				{
 					UE_LOG(LogInit, Display, TEXT("\tFound project file %s."), *GameProjectFile);
 					FPaths::SetProjectFilePath(GameProjectFile);
+
+					// Fixup command line if project file wasn't found in specified directory to properly parse next arguments.
+					FString OldCommandLine = FString(FCommandLine::Get());
+					OldCommandLine.ReplaceInline(*GameProjectFilePathUnnormalized, *GameProjectFile, ESearchCase::CaseSensitive);
+					FCommandLine::Set(*OldCommandLine);
+					CmdLine = FCommandLine::Get();
 				}
 			}
 		}
diff --git a/Engine/Source/Runtime/Launch/Private/LaunchPrivatePCH.h b/Engine/Source/Runtime/Launch/Private/LaunchPrivatePCH.h
--- a/Engine/Source/Runtime/Launch/Private/LaunchPrivatePCH.h
+++ b/Engine/Source/Runtime/Launch/Private/LaunchPrivatePCH.h
@@ -26,4 +26,4 @@
 /**
  * Attempts to set GameName based on the command line and check that we have a game name
  */
-bool LaunchSetGameName(const TCHAR *CmdLine);
+bool LaunchSetGameName(const TCHAR *CmdLine, FString& OutGameProjectFilePathUnnormalized);
