Fixed compile error in LaunchEngineLoop.cpp

diff --git a/Engine/Source/Runtime/Launch/Private/LaunchEngineLoop.cpp b/Engine/Source/Runtime/Launch/Private/LaunchEngineLoop.cpp
--- a/Engine/Source/Runtime/Launch/Private/LaunchEngineLoop.cpp
+++ b/Engine/Source/Runtime/Launch/Private/LaunchEngineLoop.cpp
@@ -2062,7 +2062,7 @@ int32 FEngineLoop::Init()
 	{
 		FModuleManager::Get().LoadModule(TEXT("ProfilerClient"));
 	}
-#endif#endif
+#endif
 
 	GIsRunning = true;
 
