Fix the order in which command-line arguments are merged together in packaged games which have a UE4CommandLine.txt file. The project name must always be the first argument (and can be followed by a map name), and is stored in UE4CommandLine.txt for packaged games but passed on the command line for Rocket (which uses UE4CommandLine.txt to pass switches).

diff --git a/Engine/Source/Runtime/Launch/Private/Windows/LaunchWindows.cpp b/Engine/Source/Runtime/Launch/Private/Windows/LaunchWindows.cpp
--- a/Engine/Source/Runtime/Launch/Private/Windows/LaunchWindows.cpp
+++ b/Engine/Source/Runtime/Launch/Private/Windows/LaunchWindows.cpp
@@ -139,6 +139,52 @@ LAUNCH_API int32 GuardedMainWrapper( const TCHAR* CmdLine, HINSTANCE hInInstance
 	return ErrorLevel;
 }
 
+static void MergeDefaultArgumentsIntoCommandLine(FString& CommandLine, FString DefaultArguments)
+{
+	// Normalize the arguments by removing any newlines and unnecessary spaces
+	DefaultArguments.ReplaceInline(TEXT("\r"), TEXT(" "));
+	DefaultArguments.ReplaceInline(TEXT("\n"), TEXT(" "));
+	DefaultArguments = DefaultArguments.Trim();
+
+	// We need to make sure that the project is always the first argument, otherwise it's not parsed correctly in LaunchEngineLoop (which also
+	// interferes with passing a map name on the command-line). It's possible that UE4CommandLine.txt contains the game name (for packaged games) 
+	// or that the command line contains the game name (for the editor), so we have to be careful in which order they're stitched together.
+	if(DefaultArguments.Len() > 0)
+	{
+		// Measure the length of the application name in the existing command line, respecting quoted substrings.
+		int32 AppNameLength = 0;
+		for(bool bInQuotes = false; AppNameLength < CommandLine.Len(); AppNameLength++)
+		{
+			if(CommandLine[AppNameLength] == '\"')
+			{
+				bInQuotes ^= true;
+			}
+			else if(!bInQuotes && FChar::IsWhitespace(CommandLine[AppNameLength]))
+			{
+				break;
+			}
+		}
+
+		// Measure the length of the project name in the default command line.
+		int32 ProjectNameLength = 0;
+		const TCHAR* RemainingArguments = *DefaultArguments;
+		FString FirstArgument;
+		if(FParse::Token(RemainingArguments, FirstArgument, false) && !FirstArgument.StartsWith("-"))
+		{
+			ProjectNameLength = RemainingArguments - *DefaultArguments;
+		}
+
+		// Build a new command line consisting of the application name, project name, command line arguments, and default arguments.
+		FString NewCommandLine;
+		NewCommandLine.Append(*CommandLine, AppNameLength);
+		NewCommandLine.AppendChar(TEXT(' '));
+		NewCommandLine.Append(*DefaultArguments, ProjectNameLength);
+		NewCommandLine.Append(*CommandLine + AppNameLength);
+		NewCommandLine.Append(*DefaultArguments + ProjectNameLength);
+		CommandLine = NewCommandLine;
+	}
+}
+	
 static const TCHAR* GetCompleteCommandLine()
 {
 	static FString CommandLine;
@@ -178,7 +224,7 @@ static const TCHAR* GetCompleteCommandLine()
 				}
 				else
 				{
-					CommandLine += TEXT(" ") + FString(UTF8_TO_TCHAR((const ANSICHAR*)Buffer.GetData())).Replace(TEXT("\r"), TEXT(" ")).Replace(TEXT("\n"), TEXT(" ")).Trim();
+					MergeDefaultArgumentsIntoCommandLine(CommandLine, FString(UTF8_TO_TCHAR((const ANSICHAR*)Buffer.GetData())));
 				}
 			}
 			CloseHandle(FileHandle);
