Always put the project name first when we merge arguments from UE4CommandLine.txt.

diff --git a/Engine/Source/Runtime/Launch/Private/Mac/LaunchMac.cpp b/Engine/Source/Runtime/Launch/Private/Mac/LaunchMac.cpp
--- a/Engine/Source/Runtime/Launch/Private/Mac/LaunchMac.cpp
+++ b/Engine/Source/Runtime/Launch/Private/Mac/LaunchMac.cpp
@@ -252,6 +252,37 @@ void EngineCrashHandler(const FGenericCrashContext& GenericContext)
 
 @end
 
+static void MergeDefaultArgumentsIntoCommandLine(FString& CommandLine, FString DefaultArguments)
+{
+	// Normalize the arguments by removing any newlines and unnecessary spaces
+	DefaultArguments.ReplaceInline(TEXT("\r"), TEXT(" "));
+	DefaultArguments.ReplaceInline(TEXT("\n"), TEXT(" "));
+	DefaultArguments = DefaultArguments.Trim();
+
+	// We need to make sure that the project is always the first argument, otherwise it's not parsed correctly in LaunchEngineLoop (which also
+	// interferes with passing a map name on the command-line). It's possible that UE4CommandLine.txt contains the game name (for packaged games) 
+	// or that the command line contains the game name (for the editor), so we have to be careful in which order they're stitched together.
+	if(DefaultArguments.Len() > 0)
+	{
+		// Measure the length of the project name in the default command line.
+		int32 ProjectNameLength = 0;
+		const TCHAR* RemainingArguments = *DefaultArguments;
+		FString FirstArgument;
+		if(FParse::Token(RemainingArguments, FirstArgument, false) && !FirstArgument.StartsWith("-"))
+		{
+			ProjectNameLength = RemainingArguments - *DefaultArguments;
+		}
+
+		// Build a new command line consisting of the application name, project name, command line arguments, and default arguments.
+		FString NewCommandLine;
+		NewCommandLine.Append(*DefaultArguments, ProjectNameLength);
+		NewCommandLine.AppendChar(TEXT(' '));
+		NewCommandLine.Append(*CommandLine);
+		NewCommandLine.Append(*DefaultArguments + ProjectNameLength);
+		CommandLine = NewCommandLine;
+	}
+}
+
 INT32_MAIN_INT32_ARGC_TCHAR_ARGV()
 {
 	for (int32 Option = 1; Option < ArgC; Option++)
@@ -294,7 +325,7 @@ INT32_MAIN_INT32_ARGC_TCHAR_ARGV()
 			int ReadSize = read(FileHandle, Buffer.GetData(), FileLength);
 			if(ReadSize == FileLength)
 			{
-				GSavedCommandLine += TEXT(" ") + FString(UTF8_TO_TCHAR((const ANSICHAR*)Buffer.GetData())).Replace(TEXT("\r"), TEXT(" ")).Replace(TEXT("\n"), TEXT(" ")).Trim();
+				MergeDefaultArgumentsIntoCommandLine(GSavedCommandLine, UTF8_TO_TCHAR((const ANSICHAR*)Buffer.GetData()));
 			}
 			else
 			{
@@ -309,4 +340,4 @@ INT32_MAIN_INT32_ARGC_TCHAR_ARGV()
 	[NSApp setDelegate:[UE4AppDelegate new]];
 	[NSApp run];
 	return 0;
-}
+}
\ No newline at end of file
