Fixed MergeDefaultArgumentsIntoCommandLine joining two flags together

UE-15377 - Right-click>Launch Game does not work for code projects

If you passed -game when running an installed engine (which adds -installedengine), you'd get "-game-installedengine" in the final command line, rather than "-game -installedengine".

diff --git a/Engine/Source/Runtime/Launch/Private/Mac/LaunchMac.cpp b/Engine/Source/Runtime/Launch/Private/Mac/LaunchMac.cpp
--- a/Engine/Source/Runtime/Launch/Private/Mac/LaunchMac.cpp
+++ b/Engine/Source/Runtime/Launch/Private/Mac/LaunchMac.cpp
@@ -278,6 +278,7 @@ static void MergeDefaultArgumentsIntoCommandLine(FString& CommandLine, FString D
 		NewCommandLine.Append(*DefaultArguments, ProjectNameLength);
 		NewCommandLine.AppendChar(TEXT(' '));
 		NewCommandLine.Append(*CommandLine);
+		NewCommandLine.AppendChar(TEXT(' '));
 		NewCommandLine.Append(*DefaultArguments + ProjectNameLength);
 		CommandLine = NewCommandLine;
 	}
diff --git a/Engine/Source/Runtime/Launch/Private/Windows/LaunchWindows.cpp b/Engine/Source/Runtime/Launch/Private/Windows/LaunchWindows.cpp
--- a/Engine/Source/Runtime/Launch/Private/Windows/LaunchWindows.cpp
+++ b/Engine/Source/Runtime/Launch/Private/Windows/LaunchWindows.cpp
@@ -180,6 +180,7 @@ static void MergeDefaultArgumentsIntoCommandLine(FString& CommandLine, FString D
 		NewCommandLine.AppendChar(TEXT(' '));
 		NewCommandLine.Append(*DefaultArguments, ProjectNameLength);
 		NewCommandLine.Append(*CommandLine + AppNameLength);
+		NewCommandLine.AppendChar(TEXT(' '));
 		NewCommandLine.Append(*DefaultArguments + ProjectNameLength);
 		CommandLine = NewCommandLine;
 	}
