fixes for UFE not working properly in a binary release
UE-15538, UE-15528, UE-15487

diff --git a/Engine/Source/Editor/UnrealEd/Private/Kismet2/DebuggerCommands.cpp b/Engine/Source/Editor/UnrealEd/Private/Kismet2/DebuggerCommands.cpp
--- a/Engine/Source/Editor/UnrealEd/Private/Kismet2/DebuggerCommands.cpp
+++ b/Engine/Source/Editor/UnrealEd/Private/Kismet2/DebuggerCommands.cpp
@@ -1567,12 +1567,12 @@ bool FInternalPlayWorldCommandCallbacks::IsReadyToLaunchOnDevice(FString DeviceI
 		{
 		case ETargetPlatformReadyStatus::CodeUnsupported:
 			// show the message
-			FMessageDialog::Open(EAppMsgType::Ok, LOCTEXT("NotSupported", "Sorry, packaging a code-based project for the selected platform is currently not supported. This feature may be available in a future release."));
+			FMessageDialog::Open(EAppMsgType::Ok, LOCTEXT("NotSupported", "Sorry, launching a code-based project for the selected platform is currently not supported. This feature may be available in a future release."));
 			return false;
 
 		case ETargetPlatformReadyStatus::PluginsUnsupported:
 			// show the message
-			FMessageDialog::Open(EAppMsgType::Ok, LOCTEXT("NotSupported", "Sorry, packaging a project with third-party plugins is currently not supported for the selected platform. This feature may be available in a future release."));
+			FMessageDialog::Open(EAppMsgType::Ok, LOCTEXT("NotSupported", "Sorry, launching a project with third-party plugins is currently not supported for the selected platform. This feature may be available in a future release."));
 			return false;
 		default:
 			break;
diff --git a/Engine/Source/Programs/UnrealFrontend/Private/Mac/MacUnrealFrontendMain.cpp b/Engine/Source/Programs/UnrealFrontend/Private/Mac/MacUnrealFrontendMain.cpp
--- a/Engine/Source/Programs/UnrealFrontend/Private/Mac/MacUnrealFrontendMain.cpp
+++ b/Engine/Source/Programs/UnrealFrontend/Private/Mac/MacUnrealFrontendMain.cpp
@@ -110,6 +110,10 @@ int main( int argc, char *argv[] )
 		GSavedCommandLine += Argument;
 	}
 
+	FCommandLine::Set(*GSavedCommandLine);
+	GSavedCommandLine = FPlatformMisc::GetCompleteCommandLine();
+
+
 	SCOPED_AUTORELEASE_POOL;
 	[NSApplication sharedApplication];
 	[NSApp setDelegate:[UE4AppDelegate new]];
diff --git a/Engine/Source/Programs/UnrealFrontend/Private/UnrealFrontendMain.cpp b/Engine/Source/Programs/UnrealFrontend/Private/UnrealFrontendMain.cpp
--- a/Engine/Source/Programs/UnrealFrontend/Private/UnrealFrontendMain.cpp
+++ b/Engine/Source/Programs/UnrealFrontend/Private/UnrealFrontendMain.cpp
@@ -20,6 +20,8 @@ IMPLEMENT_APPLICATION(UnrealFrontend, "UnrealFrontend");
  */
 int32 UnrealFrontendMain( const TCHAR* CommandLine )
 {
+	FCommandLine::Set(CommandLine);
+
 	FString Command;
 	FString Params;
 	FString NewCommandLine = CommandLine;
diff --git a/Engine/Source/Programs/UnrealFrontend/Private/Windows/WindowsUnrealFrontendMain.cpp b/Engine/Source/Programs/UnrealFrontend/Private/Windows/WindowsUnrealFrontendMain.cpp
--- a/Engine/Source/Programs/UnrealFrontend/Private/Windows/WindowsUnrealFrontendMain.cpp
+++ b/Engine/Source/Programs/UnrealFrontend/Private/Windows/WindowsUnrealFrontendMain.cpp
@@ -5,7 +5,6 @@
 #include "ExceptionHandling.h"
 #include "LaunchEngineLoop.h"
 
-
 /**
  * The main application entry point for Windows platforms.
  *
@@ -19,8 +18,7 @@ int32 WINAPI WinMain( HINSTANCE hInInstance, HINSTANCE hPrevInstance, char* lpCm
 {
 	hInstance = hInInstance;
 
-	const TCHAR* OrgCmdLine = GetCommandLineW();
-	const TCHAR* CmdLine = FCommandLine::RemoveExeName(OrgCmdLine);
+	const TCHAR* CmdLine = FPlatformMisc::GetCompleteCommandLine();
 
 #if !UE_BUILD_SHIPPING
 	if (FParse::Param(CmdLine, TEXT("crashreports")))
diff --git a/Engine/Source/Runtime/Core/Private/GenericPlatform/GenericPlatformMisc.cpp b/Engine/Source/Runtime/Core/Private/GenericPlatform/GenericPlatformMisc.cpp
--- a/Engine/Source/Runtime/Core/Private/GenericPlatform/GenericPlatformMisc.cpp
+++ b/Engine/Source/Runtime/Core/Private/GenericPlatform/GenericPlatformMisc.cpp
@@ -875,4 +875,56 @@ FString FGenericPlatformMisc::GetOperatingSystemId()
 {
 	// not implemented by default.
 	return FString();
+}
+
+void FGenericPlatformMisc::MergeDefaultArgumentsIntoCommandLine(FString& CommandLine, FString DefaultArguments)
+{
+	// Normalize the arguments by removing any newlines and unnecessary spaces
+	DefaultArguments.ReplaceInline(TEXT("\r"), TEXT(" "));
+	DefaultArguments.ReplaceInline(TEXT("\n"), TEXT(" "));
+	DefaultArguments = DefaultArguments.Trim();
+
+	// We need to make sure that the project is always the first argument, otherwise it's not parsed correctly in LaunchEngineLoop (which also
+	// interferes with passing a map name on the command-line). It's possible that UE4CommandLine.txt contains the game name (for packaged games) 
+	// or that the command line contains the game name (for the editor), so we have to be careful in which order they're stitched together.
+	if (DefaultArguments.Len() > 0)
+	{
+		// Measure the length of the application name in the existing command line, respecting quoted substrings.
+		int32 AppNameLength = 0;
+		for (bool bInQuotes = false; AppNameLength < CommandLine.Len(); AppNameLength++)
+		{
+			if (CommandLine[AppNameLength] == '\"')
+			{
+				bInQuotes ^= true;
+			}
+			else if (!bInQuotes && FChar::IsWhitespace(CommandLine[AppNameLength]))
+			{
+				break;
+			}
+		}
+
+		// Measure the length of the project name in the default command line.
+		int32 ProjectNameLength = 0;
+		const TCHAR* RemainingArguments = *DefaultArguments;
+		FString FirstArgument;
+		if (FParse::Token(RemainingArguments, FirstArgument, false) && !FirstArgument.StartsWith("-"))
+		{
+			ProjectNameLength = RemainingArguments - *DefaultArguments;
+		}
+
+		// Build a new command line consisting of the application name, project name, command line arguments, and default arguments.
+		FString NewCommandLine;
+		NewCommandLine.Append(*CommandLine, AppNameLength);
+		NewCommandLine.AppendChar(TEXT(' '));
+		NewCommandLine.Append(*DefaultArguments, ProjectNameLength);
+		NewCommandLine.Append(*CommandLine + AppNameLength);
+		NewCommandLine.AppendChar(TEXT(' '));
+		NewCommandLine.Append(*DefaultArguments + ProjectNameLength);
+		CommandLine = NewCommandLine;
+	}
+}
+
+const TCHAR* FGenericPlatformMisc::GetCompleteCommandLine()
+{
+	return FCommandLine::Get();
 }
\ No newline at end of file
diff --git a/Engine/Source/Runtime/Core/Private/Mac/MacPlatformMisc.cpp b/Engine/Source/Runtime/Core/Private/Mac/MacPlatformMisc.cpp
--- a/Engine/Source/Runtime/Core/Private/Mac/MacPlatformMisc.cpp
+++ b/Engine/Source/Runtime/Core/Private/Mac/MacPlatformMisc.cpp
@@ -1721,3 +1721,41 @@ void FMacCrashContext::GenerateCrashInfoAndLaunchReporter() const
 	
 	_Exit(0);
 }
+
+const TCHAR* FMacPlatformMisc::GetCompleteCommandLine()
+{
+    static FString CommandLine;
+    if (CommandLine.Len() == 0)
+    {
+        CommandLine = FCommandLine::Get();
+        FString CommandLineFile = FPaths::RootDir() / TEXT("UE4CommandLine.txt");
+        FPaths::NormalizeFilename(CommandLineFile);
+
+        int32 FileHandle = open(TCHAR_TO_UTF8(*CommandLineFile), O_RDONLY);
+        if (FileHandle != -1)
+        {
+            int32 FileLength = lseek(FileHandle, 0, SEEK_END);
+            if (FileLength < 0 || lseek(FileHandle, 0, SEEK_SET) != 0)
+            {
+                FPlatformMisc::LowLevelOutputDebugString(TEXT("WARNING: Failed to seek in UE4CommandLine.txt"));
+            }
+            else
+            {
+                TArray<uint8> Buffer;
+                Buffer.AddZeroed(FileLength + 1);
+
+                int ReadSize = read(FileHandle, Buffer.GetData(), FileLength);
+                if (ReadSize == FileLength)
+                {
+                    MergeDefaultArgumentsIntoCommandLine(CommandLine, UTF8_TO_TCHAR((const ANSICHAR*)Buffer.GetData()));
+                }
+                else
+                {
+                    FPlatformMisc::LowLevelOutputDebugString(TEXT("WARNING: Failed to read UE4CommandLine.txt"));
+                }
+            }
+            close(FileHandle);
+        }
+    }
+	return *CommandLine;
+}
\ No newline at end of file
diff --git a/Engine/Source/Runtime/Core/Private/Windows/WindowsPlatformMisc.cpp b/Engine/Source/Runtime/Core/Private/Windows/WindowsPlatformMisc.cpp
--- a/Engine/Source/Runtime/Core/Private/Windows/WindowsPlatformMisc.cpp
+++ b/Engine/Source/Runtime/Core/Private/Windows/WindowsPlatformMisc.cpp
@@ -2430,3 +2430,51 @@ IPlatformChunkInstall* FWindowsPlatformMisc::GetPlatformChunkInstall()
 
 	return ChunkInstall;
 }
+
+const TCHAR* FWindowsPlatformMisc::GetCompleteCommandLine()
+{
+	static FString CommandLine;
+	if (CommandLine.Len() == 0)
+	{
+		const TCHAR* OrgCmdLine = ::GetCommandLine();
+		CommandLine = FCommandLine::RemoveExeName(OrgCmdLine);
+
+		// Get the path to the arguments file
+		FString ArgsFileName = FPaths::RootDir() / TEXT("UE4CommandLine.txt");
+
+		// Try to read the default arguments
+		HANDLE FileHandle = CreateFile(*ArgsFileName, GENERIC_READ, FILE_SHARE_READ | FILE_SHARE_DELETE, NULL, OPEN_EXISTING, FILE_ATTRIBUTE_NORMAL, NULL);
+		if (FileHandle == INVALID_HANDLE_VALUE)
+		{
+			if (GetLastError() != ERROR_FILE_NOT_FOUND)
+			{
+				FPlatformMisc::LowLevelOutputDebugString(*FString::Printf(TEXT("WARNING: Failed to open UE4CommandLine.txt - %08X\n"), GetLastError()));
+			}
+		}
+		else
+		{
+			::DWORD FileSize = ::GetFileSize(FileHandle, NULL);
+			if (FileSize == (::DWORD)0xffffffff)
+			{
+				FPlatformMisc::LowLevelOutputDebugString(*FString::Printf(TEXT("WARNING: Failed to get file size for UE4CommandLine.txt (%08x)\n"), GetLastError()));
+			}
+			else
+			{
+				TArray<BYTE> Buffer;
+				Buffer.AddZeroed(FileSize + 1);
+
+				::DWORD FileSizeRead = 0;
+				if (!ReadFile(FileHandle, Buffer.GetData(), FileSize, &FileSizeRead, NULL) || FileSizeRead < FileSize)
+				{
+					FPlatformMisc::LowLevelOutputDebugString(*FString::Printf(TEXT("WARNING: Failed to read UE4CommandLine.txt file (%08x)\n"), GetLastError()));
+				}
+				else
+				{
+					MergeDefaultArgumentsIntoCommandLine(CommandLine, FString(UTF8_TO_TCHAR((const ANSICHAR*)Buffer.GetData())));
+				}
+			}
+			CloseHandle(FileHandle);
+		}
+	}
+	return *CommandLine;
+}
diff --git a/Engine/Source/Runtime/Core/Public/GenericPlatform/GenericPlatformMisc.h b/Engine/Source/Runtime/Core/Public/GenericPlatform/GenericPlatformMisc.h
--- a/Engine/Source/Runtime/Core/Public/GenericPlatform/GenericPlatformMisc.h
+++ b/Engine/Source/Runtime/Core/Public/GenericPlatform/GenericPlatformMisc.h
@@ -835,6 +835,11 @@ public:
 	static const void PreLoadMap(FString&, FString&, void*)
 	{}
 
+	/**
+	 * Get a completed commandline which includes reading UE4CommandLine.txt from root
+	 */
+	static const TCHAR* GetCompleteCommandLine();
+
 #if !UE_BUILD_SHIPPING
 protected:
 	/** Whether the user should be prompted to allow for a remote debugger to be attached */
@@ -842,4 +847,9 @@ protected:
 	/** Whether the user should be prompted to allow for a remote debugger to be attached on an ensure */
 	static bool bPromptForRemoteDebugOnEnsure;
 #endif	//#if !UE_BUILD_SHIPPING
+
+	/**
+	 * Merge the additional arguments in to the commandline
+	 */
+	static void MergeDefaultArgumentsIntoCommandLine(FString& CommandLine, FString DefaultArguments);
 };
diff --git a/Engine/Source/Runtime/Core/Public/Mac/MacPlatformMisc.h b/Engine/Source/Runtime/Core/Public/Mac/MacPlatformMisc.h
--- a/Engine/Source/Runtime/Core/Public/Mac/MacPlatformMisc.h
+++ b/Engine/Source/Runtime/Core/Public/Mac/MacPlatformMisc.h
@@ -201,6 +201,8 @@ struct CORE_API FMacPlatformMisc : public FGenericPlatformMisc
 	static bool bChachedMacMenuStateNeedsUpdate;
 
 	static id<NSObject> CommandletActivity;
+
+	static const TCHAR* GetCompleteCommandLine();
 };
 
 #ifdef __OBJC__
diff --git a/Engine/Source/Runtime/Core/Public/Windows/WindowsPlatformMisc.h b/Engine/Source/Runtime/Core/Public/Windows/WindowsPlatformMisc.h
--- a/Engine/Source/Runtime/Core/Public/Windows/WindowsPlatformMisc.h
+++ b/Engine/Source/Runtime/Core/Public/Windows/WindowsPlatformMisc.h
@@ -263,6 +263,8 @@ struct CORE_API FWindowsPlatformMisc
 	static EConvertibleLaptopMode GetConvertibleLaptopMode();
 
 	static IPlatformChunkInstall* GetPlatformChunkInstall();
+
+	static const TCHAR* GetCompleteCommandLine();
 };
 
 
diff --git a/Engine/Source/Runtime/Launch/Private/Mac/LaunchMac.cpp b/Engine/Source/Runtime/Launch/Private/Mac/LaunchMac.cpp
--- a/Engine/Source/Runtime/Launch/Private/Mac/LaunchMac.cpp
+++ b/Engine/Source/Runtime/Launch/Private/Mac/LaunchMac.cpp
@@ -252,38 +252,6 @@ void EngineCrashHandler(const FGenericCrashContext& GenericContext)
 
 @end
 
-static void MergeDefaultArgumentsIntoCommandLine(FString& CommandLine, FString DefaultArguments)
-{
-	// Normalize the arguments by removing any newlines and unnecessary spaces
-	DefaultArguments.ReplaceInline(TEXT("\r"), TEXT(" "));
-	DefaultArguments.ReplaceInline(TEXT("\n"), TEXT(" "));
-	DefaultArguments = DefaultArguments.Trim();
-
-	// We need to make sure that the project is always the first argument, otherwise it's not parsed correctly in LaunchEngineLoop (which also
-	// interferes with passing a map name on the command-line). It's possible that UE4CommandLine.txt contains the game name (for packaged games) 
-	// or that the command line contains the game name (for the editor), so we have to be careful in which order they're stitched together.
-	if(DefaultArguments.Len() > 0)
-	{
-		// Measure the length of the project name in the default command line.
-		int32 ProjectNameLength = 0;
-		const TCHAR* RemainingArguments = *DefaultArguments;
-		FString FirstArgument;
-		if(FParse::Token(RemainingArguments, FirstArgument, false) && !FirstArgument.StartsWith("-"))
-		{
-			ProjectNameLength = RemainingArguments - *DefaultArguments;
-		}
-
-		// Build a new command line consisting of the application name, project name, command line arguments, and default arguments.
-		FString NewCommandLine;
-		NewCommandLine.Append(*DefaultArguments, ProjectNameLength);
-		NewCommandLine.AppendChar(TEXT(' '));
-		NewCommandLine.Append(*CommandLine);
-		NewCommandLine.AppendChar(TEXT(' '));
-		NewCommandLine.Append(*DefaultArguments + ProjectNameLength);
-		CommandLine = NewCommandLine;
-	}
-}
-
 INT32_MAIN_INT32_ARGC_TCHAR_ARGV()
 {
 	for (int32 Option = 1; Option < ArgC; Option++)
@@ -306,35 +274,9 @@ INT32_MAIN_INT32_ARGC_TCHAR_ARGV()
 		}
 		GSavedCommandLine += Argument;
 	}
-	
-	FString CommandLineFile = FPaths::RootDir() / TEXT("UE4CommandLine.txt");
-	FPaths::NormalizeFilename(CommandLineFile);
-
-	int32 FileHandle = open(TCHAR_TO_UTF8(*CommandLineFile), O_RDONLY);
-	if (FileHandle != -1)
-	{
-		int32 FileLength = lseek(FileHandle, 0, SEEK_END);
-		if(FileLength < 0 || lseek(FileHandle, 0, SEEK_SET) != 0)
-		{
-			FPlatformMisc::LowLevelOutputDebugString(TEXT("WARNING: Failed to seek in UE4CommandLine.txt"));
-		}
-		else
-		{
-			TArray<uint8> Buffer;
-			Buffer.AddZeroed(FileLength + 1);
 
-			int ReadSize = read(FileHandle, Buffer.GetData(), FileLength);
-			if(ReadSize == FileLength)
-			{
-				MergeDefaultArgumentsIntoCommandLine(GSavedCommandLine, UTF8_TO_TCHAR((const ANSICHAR*)Buffer.GetData()));
-			}
-			else
-			{
-				FPlatformMisc::LowLevelOutputDebugString(TEXT("WARNING: Failed to read UE4CommandLine.txt"));
-			}
-		}
-		close(FileHandle);
-	}
+	FCommandLine::Set(*GSavedCommandLine);
+	GSavedCommandLine = FPlatformMisc::GetCompleteCommandLine();
 
 	SCOPED_AUTORELEASE_POOL;
 	[NSApplication sharedApplication];
diff --git a/Engine/Source/Runtime/Launch/Private/Windows/LaunchWindows.cpp b/Engine/Source/Runtime/Launch/Private/Windows/LaunchWindows.cpp
--- a/Engine/Source/Runtime/Launch/Private/Windows/LaunchWindows.cpp
+++ b/Engine/Source/Runtime/Launch/Private/Windows/LaunchWindows.cpp
@@ -138,102 +138,6 @@ LAUNCH_API int32 GuardedMainWrapper( const TCHAR* CmdLine, HINSTANCE hInInstance
 	}
 	return ErrorLevel;
 }
-
-static void MergeDefaultArgumentsIntoCommandLine(FString& CommandLine, FString DefaultArguments)
-{
-	// Normalize the arguments by removing any newlines and unnecessary spaces
-	DefaultArguments.ReplaceInline(TEXT("\r"), TEXT(" "));
-	DefaultArguments.ReplaceInline(TEXT("\n"), TEXT(" "));
-	DefaultArguments = DefaultArguments.Trim();
-
-	// We need to make sure that the project is always the first argument, otherwise it's not parsed correctly in LaunchEngineLoop (which also
-	// interferes with passing a map name on the command-line). It's possible that UE4CommandLine.txt contains the game name (for packaged games) 
-	// or that the command line contains the game name (for the editor), so we have to be careful in which order they're stitched together.
-	if(DefaultArguments.Len() > 0)
-	{
-		// Measure the length of the application name in the existing command line, respecting quoted substrings.
-		int32 AppNameLength = 0;
-		for(bool bInQuotes = false; AppNameLength < CommandLine.Len(); AppNameLength++)
-		{
-			if(CommandLine[AppNameLength] == '\"')
-			{
-				bInQuotes ^= true;
-			}
-			else if(!bInQuotes && FChar::IsWhitespace(CommandLine[AppNameLength]))
-			{
-				break;
-			}
-		}
-
-		// Measure the length of the project name in the default command line.
-		int32 ProjectNameLength = 0;
-		const TCHAR* RemainingArguments = *DefaultArguments;
-		FString FirstArgument;
-		if(FParse::Token(RemainingArguments, FirstArgument, false) && !FirstArgument.StartsWith("-"))
-		{
-			ProjectNameLength = RemainingArguments - *DefaultArguments;
-		}
-
-		// Build a new command line consisting of the application name, project name, command line arguments, and default arguments.
-		FString NewCommandLine;
-		NewCommandLine.Append(*CommandLine, AppNameLength);
-		NewCommandLine.AppendChar(TEXT(' '));
-		NewCommandLine.Append(*DefaultArguments, ProjectNameLength);
-		NewCommandLine.Append(*CommandLine + AppNameLength);
-		NewCommandLine.AppendChar(TEXT(' '));
-		NewCommandLine.Append(*DefaultArguments + ProjectNameLength);
-		CommandLine = NewCommandLine;
-	}
-}
-	
-static const TCHAR* GetCompleteCommandLine()
-{
-	static FString CommandLine;
-	if(CommandLine.Len() == 0)
-	{
-		// Default to the native command line
-		CommandLine = ::GetCommandLine();
-
-		// Get the path to the arguments file
-		FString ArgsFileName = FPaths::RootDir() / TEXT("UE4CommandLine.txt");
-
-		// Try to read the default arguments
-		HANDLE FileHandle = CreateFile(*ArgsFileName, GENERIC_READ, FILE_SHARE_READ | FILE_SHARE_DELETE, NULL, OPEN_EXISTING, FILE_ATTRIBUTE_NORMAL, NULL);
-		if(FileHandle == INVALID_HANDLE_VALUE)
-		{
-			if(GetLastError() != ERROR_FILE_NOT_FOUND)
-			{
-				FPlatformMisc::LowLevelOutputDebugString(*FString::Printf(TEXT("WARNING: Failed to open UE4CommandLine.txt - %08X\n"), GetLastError()));
-			}
-		}
-		else
-		{
-			::DWORD FileSize = ::GetFileSize(FileHandle, NULL);
-			if(FileSize == (::DWORD)0xffffffff)
-			{
-				FPlatformMisc::LowLevelOutputDebugString(*FString::Printf(TEXT("WARNING: Failed to get file size for UE4CommandLine.txt (%08x)\n"), GetLastError()));
-			}
-			else
-			{
-				TArray<BYTE> Buffer;
-				Buffer.AddZeroed(FileSize + 1);
-
-				::DWORD FileSizeRead = 0;
-				if (!ReadFile(FileHandle, Buffer.GetData(), FileSize, &FileSizeRead, NULL) || FileSizeRead < FileSize)
-				{
-					FPlatformMisc::LowLevelOutputDebugString(*FString::Printf(TEXT("WARNING: Failed to read UE4CommandLine.txt file (%08x)\n"), GetLastError()));
-				}
-				else
-				{
-					MergeDefaultArgumentsIntoCommandLine(CommandLine, FString(UTF8_TO_TCHAR((const ANSICHAR*)Buffer.GetData())));
-				}
-			}
-			CloseHandle(FileHandle);
-		}
-	}
-	return *CommandLine;
-}
-
 int32 WINAPI WinMain( HINSTANCE hInInstance, HINSTANCE hPrevInstance, char*, int32 nCmdShow )
 {
 	// Setup common Windows settings
@@ -241,7 +145,7 @@ int32 WINAPI WinMain( HINSTANCE hInInstance, HINSTANCE hPrevInstance, char*, int
 
 	int32 ErrorLevel			= 0;
 	hInstance				= hInInstance;
-	const TCHAR* CmdLine = GetCompleteCommandLine();
+	const TCHAR* CmdLine = FPlatformMisc::GetCompleteCommandLine();
 
 #if !(UE_BUILD_SHIPPING && WITH_EDITOR)
 	// Named mutex we use to figure out whether we are the first instance of the game running. This is needed to e.g.
