- key off of manifest for whether to enable the GearVR plugin, enabling blueprint only projects out of the box
- fix for lollipop issue

diff --git a/Engine/Build/Android/Java/src/com/epicgames/ue4/GameActivity.java b/Engine/Build/Android/Java/src/com/epicgames/ue4/GameActivity.java
--- a/Engine/Build/Android/Java/src/com/epicgames/ue4/GameActivity.java
+++ b/Engine/Build/Android/Java/src/com/epicgames/ue4/GameActivity.java
@@ -146,6 +146,9 @@ public class GameActivity extends NativeActivity
 	
 	/** Check to see if we should be verifying the files once we have them */
 	public boolean VerifyOBBOnStartUp = false;
+
+	/** Whether this application was packaged for GearVR or not */
+	public boolean PackagedForGearVR = false;
 	
 	/** Flag to ensure we have finished startup before allowing nativeOnActivityResult to get called */
 	private boolean InitCompletedOK = false;
@@ -363,8 +366,18 @@ public class GameActivity extends NativeActivity
 			{
 				Log.debug( "UI hiding not found. Leaving as " + ShouldHideUI);
 			}
-				
 			
+			if(bundle.containsKey("com.samsung.android.vr.application.mode"))
+			{
+				PackagedForGearVR = true;
+				String VRMode = bundle.getString("com.samsung.android.vr.application.mode");
+				Log.debug("Found GearVR mode = " + VRMode);
+			}
+			else
+			{
+				PackagedForGearVR = false;
+				Log.debug("No GearVR mode detected.");
+			}
 		}
 		catch (NameNotFoundException e)
 		{
@@ -943,6 +956,12 @@ public class GameActivity extends NativeActivity
 		});
 	}
 
+	// check the manifest to determine if we are a GearVR application
+	public boolean AndroidThunkJava_IsGearVRApplication()
+	{
+		return PackagedForGearVR;
+	}
+
 	public static String AndroidThunkJava_GetFontDirectory()
 	{
 		// Parse and find the first known fonts directory on the device
diff --git a/Engine/Plugins/Runtime/GearVR/GearVR.uplugin b/Engine/Plugins/Runtime/GearVR/GearVR.uplugin
--- a/Engine/Plugins/Runtime/GearVR/GearVR.uplugin
+++ b/Engine/Plugins/Runtime/GearVR/GearVR.uplugin
@@ -9,7 +9,7 @@
 	"EngineVersion" : "4.2.0",
 	"Description" : "Support for the Gear VR head mounted display",
 	"Category" : "Virtual Reality",
-	"EnabledByDefault" : false,
+	"EnabledByDefault" : true,
 
 	"Modules" :
 	[
diff --git a/Engine/Plugins/Runtime/GearVR/Source/GearVR/Private/GearVR.cpp b/Engine/Plugins/Runtime/GearVR/Source/GearVR/Private/GearVR.cpp
--- a/Engine/Plugins/Runtime/GearVR/Source/GearVR/Private/GearVR.cpp
+++ b/Engine/Plugins/Runtime/GearVR/Source/GearVR/Private/GearVR.cpp
@@ -27,7 +27,7 @@ do \
 int err; \
 	while ((err = glGetError()) != GL_NO_ERROR) \
 	{ \
-UE_LOG(LogHMD, Warning, TEXT("%s:%d GL error (%#x)\n"), ANSI_TO_TCHAR(__FILE__), __LINE__, err); \
+		UE_LOG(LogHMD, Warning, TEXT("%s:%d GL error (#%x)\n"), ANSI_TO_TCHAR(__FILE__), __LINE__, err); \
 	} \
 } while (0) 
 
@@ -36,6 +36,10 @@ UE_LOG(LogHMD, Warning, TEXT("%s:%d GL error (%#x)\n"), ANSI_TO_TCHAR(__FILE__),
 #endif // #if !UE_BUILD_SHIPPING
 
 
+// call out to JNI to see if the application was packaged for GearVR
+extern bool AndroidThunkCpp_IsGearVRApplication();
+
+
 //---------------------------------------------------
 // GearVR Plugin Implementation
 //---------------------------------------------------
@@ -59,6 +63,12 @@ IMPLEMENT_MODULE( FGearVRPlugin, GearVR )
 TSharedPtr< class IHeadMountedDisplay, ESPMode::ThreadSafe > FGearVRPlugin::CreateHeadMountedDisplay()
 {
 #if GEARVR_SUPPORTED_PLATFORMS
+	if (!AndroidThunkCpp_IsGearVRApplication())
+	{
+		// don't do anything if we aren't packaged for GearVR
+		return NULL;
+	}
+
 	TSharedPtr< FGearVR, ESPMode::ThreadSafe > GearVR(new FGearVR());
 	if (GearVR->IsInitialized())
 	{
@@ -71,6 +81,12 @@ TSharedPtr< class IHeadMountedDisplay, ESPMode::ThreadSafe > FGearVRPlugin::Crea
 void FGearVRPlugin::PreInit()
 {
 #if GEARVR_SUPPORTED_PLATFORMS
+	if (!AndroidThunkCpp_IsGearVRApplication())
+	{
+		// don't do anything if we aren't packaged for GearVR
+		return;
+	}
+
 	FGearVR::PreInit();
 #endif//GEARVR_SUPPORTED_PLATFORMS
 }
diff --git a/Engine/Source/Runtime/Launch/Private/Android/AndroidJNI.cpp b/Engine/Source/Runtime/Launch/Private/Android/AndroidJNI.cpp
--- a/Engine/Source/Runtime/Launch/Private/Android/AndroidJNI.cpp
+++ b/Engine/Source/Runtime/Launch/Private/Android/AndroidJNI.cpp
@@ -52,6 +52,7 @@ void FJavaWrapper::FindClassesAndMethods(JNIEnv* Env)
 	AndroidThunkJava_IsMusicActive = FindMethod(Env, GameActivityClassID, "AndroidThunkJava_IsMusicActive", "()Z", bIsOptional);
 	AndroidThunkJava_KeepScreenOn = FindMethod(Env, GameActivityClassID, "AndroidThunkJava_KeepScreenOn", "(Z)V", bIsOptional);
 	AndroidThunkJava_InitHMDs = FindMethod(Env, GameActivityClassID, "AndroidThunkJava_InitHMDs", "()V", bIsOptional);
+	AndroidThunkJava_IsGearVRApplication = FindMethod(Env, GameActivityClassID, "AndroidThunkJava_IsGearVRApplication", "()Z", bIsOptional);
 
 
 	// the rest are optional
@@ -157,6 +158,7 @@ jmethodID FJavaWrapper::AndroidThunkJava_Vibrate;
 jmethodID FJavaWrapper::AndroidThunkJava_IsMusicActive;
 jmethodID FJavaWrapper::AndroidThunkJava_KeepScreenOn;
 jmethodID FJavaWrapper::AndroidThunkJava_InitHMDs;
+jmethodID FJavaWrapper::AndroidThunkJava_IsGearVRApplication;
 
 jclass FJavaWrapper::GoogleServicesClassID;
 jobject FJavaWrapper::GoogleServicesThis;
@@ -218,6 +220,18 @@ void AndroidThunkCpp_InitHMDs()
 	}
 }
 
+// Check to see if this application is packaged for GearVR
+bool AndroidThunkCpp_IsGearVRApplication()
+{
+	bool bIsGearVRApplication = false;
+	if (JNIEnv* Env = FAndroidApplication::GetJavaEnv())
+	{
+		bIsGearVRApplication = FJavaWrapper::CallBooleanMethod(Env, FJavaWrapper::GameActivityThis, FJavaWrapper::AndroidThunkJava_IsGearVRApplication);
+	}
+
+	return bIsGearVRApplication;
+}
+
 void AndroidThunkCpp_ShowConsoleWindow()
 {
 	if (JNIEnv* Env = FAndroidApplication::GetJavaEnv())
diff --git a/Engine/Source/Runtime/Launch/Private/Android/AndroidJNI.h b/Engine/Source/Runtime/Launch/Private/Android/AndroidJNI.h
--- a/Engine/Source/Runtime/Launch/Private/Android/AndroidJNI.h
+++ b/Engine/Source/Runtime/Launch/Private/Android/AndroidJNI.h
@@ -28,6 +28,7 @@ public:
 	static jmethodID AndroidThunkJava_IsMusicActive;
 	static jmethodID AndroidThunkJava_KeepScreenOn;
 	static jmethodID AndroidThunkJava_InitHMDs;
+	static jmethodID AndroidThunkJava_IsGearVRApplication;
 
 	// IDs related to google play services
 	static jclass GoogleServicesClassID;
