Fixed the initialization of the texture streaming system so that it's initialized after RHIInit() like it's supposed to be.

diff --git a/Engine/Source/Runtime/Launch/Private/LaunchEngineLoop.cpp b/Engine/Source/Runtime/Launch/Private/LaunchEngineLoop.cpp
--- a/Engine/Source/Runtime/Launch/Private/LaunchEngineLoop.cpp
+++ b/Engine/Source/Runtime/Launch/Private/LaunchEngineLoop.cpp
@@ -1298,8 +1298,6 @@ int32 FEngineLoop::PreInit( const TCHAR* CmdLine )
 
 	EndInitTextLocalization();
 
-	IStreamingManager::Get();
-
 	if (FPlatformProcess::SupportsMultithreading() && !IsRunningDedicatedServer() && (bIsRegularClient || bHasEditorToken))
 	{
 		FPlatformSplash::Show();
@@ -1338,6 +1336,9 @@ int32 FEngineLoop::PreInit( const TCHAR* CmdLine )
 		GDistanceFieldAsyncQueue = new FDistanceFieldAsyncQueue();
 	}
 
+	// Initialize the texture streaming system (needs to happen after RHIInit).
+	IStreamingManager::Get();
+
 	{
 		DECLARE_SCOPE_CYCLE_COUNTER(TEXT("Initial UObject load"), STAT_InitialUObjectLoad, STATGROUP_LoadTime);
 
