Add logging for IsRocket() back into LaunchEngineLoop. We don't have a -Rocket command line argument any more.

diff --git a/Engine/Source/Runtime/Launch/Private/LaunchEngineLoop.cpp b/Engine/Source/Runtime/Launch/Private/LaunchEngineLoop.cpp
--- a/Engine/Source/Runtime/Launch/Private/LaunchEngineLoop.cpp
+++ b/Engine/Source/Runtime/Launch/Private/LaunchEngineLoop.cpp
@@ -2788,6 +2788,7 @@ bool FEngineLoop::AppInit( )
 	UE_LOG(LogInit, Log, TEXT("Command line: %s"), FCommandLine::Get() );
 	UE_LOG(LogInit, Log, TEXT("Base directory: %s"), FPlatformProcess::BaseDir() );
 	//UE_LOG(LogInit, Log, TEXT("Character set: %s"), sizeof(TCHAR)==1 ? TEXT("ANSI") : TEXT("Unicode") );
+	UE_LOG(LogInit, Log, TEXT("Rocket: %d"), FRocketSupport::IsRocket()? 1 : 0);
 
 	// if a logging build, clear out old log files
 #if !NO_LOGGING && !(UE_BUILD_SHIPPING || UE_BUILD_TEST)
