remove the splash screen if it is showing and an alert is popped up
#ios

diff --git a/Engine/Source/Runtime/Launch/Private/IOS/IOSAppDelegateConsoleHandling.cpp b/Engine/Source/Runtime/Launch/Private/IOS/IOSAppDelegateConsoleHandling.cpp
--- a/Engine/Source/Runtime/Launch/Private/IOS/IOSAppDelegateConsoleHandling.cpp
+++ b/Engine/Source/Runtime/Launch/Private/IOS/IOSAppDelegateConsoleHandling.cpp
@@ -4,6 +4,8 @@
 #include "IOSAppDelegateConsoleHandling.h"
 #include "LaunchPrivatePCH.h"
 
+extern bool GShowSplashScreen;
+
 @implementation IOSAppDelegate (ConsoleHandling)
 
 #if !UE_BUILD_SHIPPING
@@ -148,6 +150,14 @@
  */
 - (void)ShowAlert:(NSMutableArray*)StringArray
 {
+	if (GShowSplashScreen)
+	{
+		if ([[IOSAppDelegate GetDelegate].Window viewWithTag : 2] != nil)
+		{
+			[[[IOSAppDelegate GetDelegate].Window viewWithTag : 2] removeFromSuperview];
+		}
+		GShowSplashScreen = false;
+	}
 #ifdef __IPHONE_8_0
 	if ([UIAlertController class])
 	{
