[INTEGRATE] Change 2574438 by [REDACTED_PERSON][REDACTED_HANDLE] on 2015/06/02 16:37:58

	Add a define for whether the current engine version is a promoted build. UnrealGameSync will override this and force it to 0.

diff --git a/Engine/Source/Runtime/Launch/Resources/Version.h b/Engine/Source/Runtime/Launch/Resources/Version.h
--- a/Engine/Source/Runtime/Launch/Resources/Version.h
+++ b/Engine/Source/Runtime/Launch/Resources/Version.h
@@ -10,6 +10,8 @@
 
 #define ENGINE_IS_LICENSEE_VERSION 0
 
+#define ENGINE_IS_PROMOTED_BUILD (ENGINE_VERSION > 0)
+
 #define BUILT_FROM_CHANGELIST 0
 #define BRANCH_NAME "UE4"
 
