Stats - More stats for memory profiling

diff --git a/Engine/Source/Runtime/Launch/Private/LaunchEngineLoop.cpp b/Engine/Source/Runtime/Launch/Private/LaunchEngineLoop.cpp
--- a/Engine/Source/Runtime/Launch/Private/LaunchEngineLoop.cpp
+++ b/Engine/Source/Runtime/Launch/Private/LaunchEngineLoop.cpp
@@ -1086,6 +1086,8 @@ int32 FEngineLoop::PreInit( const TCHAR* CmdLine )
 	FThreadStats::StartThread();
 #endif
 
+	DECLARE_SCOPE_CYCLE_COUNTER( TEXT( "FEngineLoop::PreInit.AfterStats" ), STAT_FEngineLoop_PreInit_AfterStats, STATGROUP_LoadTime );
+
 	if (FPlatformProcess::SupportsMultithreading())
 	{
 		GThreadPool	= FQueuedThreadPool::Allocate();
@@ -2012,6 +2014,8 @@ void GameLoopIsStarved()
 
 int32 FEngineLoop::Init()
 {
+	DECLARE_SCOPE_CYCLE_COUNTER( TEXT( "FEngineLoop::Init" ), STAT_FEngineLoop_Init, STATGROUP_LoadTime );
+
 	FScopedSlowTask SlowTask(100);
 	SlowTask.EnterProgressFrame(10);
 
@@ -2176,6 +2180,7 @@ void FEngineLoop::Exit()
 	// order they were loaded in, so that systems can unregister and perform general clean up.
 	FModuleManager::Get().UnloadModulesAtShutdown();
 
+	// Move earlier?
 #if STATS
 	FThreadStats::StopThread();
 #endif
