revised when shared template resources are mounted

diff --git a/Engine/Source/Editor/UnrealEd/Private/UnrealEdMisc.cpp b/Engine/Source/Editor/UnrealEd/Private/UnrealEdMisc.cpp
--- a/Engine/Source/Editor/UnrealEd/Private/UnrealEdMisc.cpp
+++ b/Engine/Source/Editor/UnrealEd/Private/UnrealEdMisc.cpp
@@ -256,9 +256,6 @@ void FUnrealEdMisc::OnInit()
 	// Check for automated build/submit option
 	const bool bDoAutomatedMapBuild = FParse::Param( ParsedCmdLine, TEXT("AutomatedMapBuild") );
 
-	// This adds mount points for the template shared resources if we are editing a template project.
-	MountTemplateSharedPaths();
-
 	// Load startup map (conditionally)
 	SlowTask.EnterProgressFrame(60);
 	{
diff --git a/Engine/Source/Runtime/Launch/Private/LaunchEngineLoop.cpp b/Engine/Source/Runtime/Launch/Private/LaunchEngineLoop.cpp
--- a/Engine/Source/Runtime/Launch/Private/LaunchEngineLoop.cpp
+++ b/Engine/Source/Runtime/Launch/Private/LaunchEngineLoop.cpp
@@ -1461,6 +1461,11 @@ int32 FEngineLoop::PreInit( const TCHAR* CmdLine )
 		StartRenderingThread();
 	}
 	
+#if WITH_EDITOR
+	// We need to mount the shared resources for templates (if there are any) before we try and load and game classes
+	FUnrealEdMisc::Get().MountTemplateSharedPaths();
+#endif
+
 	if ( !LoadStartupModules() )
 	{
 		// At least one startup module failed to load, return 1 to indicate an error
