Stats - Moved stats/taskgraph to earlier location, before loading CoreUObject
Stats - Tweaked a few stats

diff --git a/Engine/Source/Runtime/Launch/Private/LaunchEngineLoop.cpp b/Engine/Source/Runtime/Launch/Private/LaunchEngineLoop.cpp
--- a/Engine/Source/Runtime/Launch/Private/LaunchEngineLoop.cpp
+++ b/Engine/Source/Runtime/Launch/Private/LaunchEngineLoop.cpp
@@ -1049,6 +1049,16 @@ int32 FEngineLoop::PreInit( const TCHAR* CmdLine )
 	}
 #endif
 
+	// initialize task graph sub-system with potential multiple threads
+	FTaskGraphInterface::Startup( FPlatformMisc::NumberOfCores() );
+	FTaskGraphInterface::Get().AttachToThread( ENamedThreads::GameThread );
+
+#if STATS
+	FThreadStats::StartThread();
+#endif
+
+	DECLARE_SCOPE_CYCLE_COUNTER( TEXT( "FEngineLoop::PreInit.AfterStats" ), STAT_FEngineLoop_PreInit_AfterStats, STATGROUP_LoadTime );
+
 	// Load Core modules required for everything else to work (needs to be loaded before InitializeRenderingCVarsCaching)
 	LoadCoreModules();
 
@@ -1078,16 +1088,6 @@ int32 FEngineLoop::PreInit( const TCHAR* CmdLine )
 	}
 #endif //WITH_EDITOR
 
-	// initialize task graph sub-system with potential multiple threads
-	FTaskGraphInterface::Startup( FPlatformMisc::NumberOfCores() );
-	FTaskGraphInterface::Get().AttachToThread( ENamedThreads::GameThread );
-
-#if STATS
-	FThreadStats::StartThread();
-#endif
-
-	DECLARE_SCOPE_CYCLE_COUNTER( TEXT( "FEngineLoop::PreInit.AfterStats" ), STAT_FEngineLoop_PreInit_AfterStats, STATGROUP_LoadTime );
-
 	if (FPlatformProcess::SupportsMultithreading())
 	{
 		GThreadPool	= FQueuedThreadPool::Allocate();
@@ -2099,6 +2099,8 @@ int32 FEngineLoop::Init()
 
 void FEngineLoop::Exit()
 {
+	STAT_ADD_CUSTOMMESSAGE_NAME( STAT_NamedMarker, TEXT( "EngineLoop.Exit" ) );
+
 	GIsRunning	= 0;
 	GLogConsole	= NULL;
 
@@ -2515,15 +2517,15 @@ void FEngineLoop::Tick()
 			RHICmdList.EndFrame();
 			RHICmdList.PopEvent();
 		});
-	} 
 
-	// Check for async platform hardware survey results
-	GEngine->TickHardwareSurvey();
+		// Check for async platform hardware survey results
+		GEngine->TickHardwareSurvey();
 
-	// Set CPU utilization stats.
-	const FCPUTime CPUTime = FPlatformTime::GetCPUTime();
-	SET_FLOAT_STAT(STAT_CPUTimePct,CPUTime.CPUTimePct);
-	SET_FLOAT_STAT(STAT_CPUTimePctRelative,CPUTime.CPUTimePctRelative);
+		// Set CPU utilization stats.
+		const FCPUTime CPUTime = FPlatformTime::GetCPUTime();
+		SET_FLOAT_STAT( STAT_CPUTimePct, CPUTime.CPUTimePct );
+		SET_FLOAT_STAT( STAT_CPUTimePctRelative, CPUTime.CPUTimePctRelative );
+	} 
 }
 
 void FEngineLoop::ClearPendingCleanupObjects()
