Stats - Fixed a bad scope for STAT_FEngineLoop_Tick_AutomationController

diff --git a/Engine/Source/Runtime/Launch/Private/LaunchEngineLoop.cpp b/Engine/Source/Runtime/Launch/Private/LaunchEngineLoop.cpp
--- a/Engine/Source/Runtime/Launch/Private/LaunchEngineLoop.cpp
+++ b/Engine/Source/Runtime/Launch/Private/LaunchEngineLoop.cpp
@@ -2431,14 +2431,16 @@ void FEngineLoop::Tick()
 #endif
 
 #if WITH_EDITOR
+		{
 #if WITH_ENGINE
-		QUICK_SCOPE_CYCLE_COUNTER(STAT_FEngineLoop_Tick_AutomationController);
+			QUICK_SCOPE_CYCLE_COUNTER( STAT_FEngineLoop_Tick_AutomationController );
 #endif
-		static FName AutomationController("AutomationController");
-		//Check if module loaded to support the change to allow this to be hot compilable.
-		if (FModuleManager::Get().IsModuleLoaded(AutomationController))
-		{
-			FModuleManager::GetModuleChecked<IAutomationControllerModule>(AutomationController).Tick();
+			static FName AutomationController( "AutomationController" );
+			//Check if module loaded to support the change to allow this to be hot compilable.
+			if (FModuleManager::Get().IsModuleLoaded( AutomationController ))
+			{
+				FModuleManager::GetModuleChecked<IAutomationControllerModule>( AutomationController ).Tick();
+			}
 		}
 #endif
 
