Fix introTutorial and blutility module loading too early in the editor case.

diff --git a/Engine/Source/Editor/UnrealEd/Private/Commandlets/GenerateBlueprintAPICommandlet.cpp b/Engine/Source/Editor/UnrealEd/Private/Commandlets/GenerateBlueprintAPICommandlet.cpp
--- a/Engine/Source/Editor/UnrealEd/Private/Commandlets/GenerateBlueprintAPICommandlet.cpp
+++ b/Engine/Source/Editor/UnrealEd/Private/Commandlets/GenerateBlueprintAPICommandlet.cpp
@@ -963,6 +963,8 @@ int32 UGenerateBlueprintAPICommandlet::Main(FString const& Params)
 	GenerateBlueprintAPIUtils::CommandOptions = GenerateBlueprintAPIUtils::CommandletOptions(Switches);
 	GenerateBlueprintAPIUtils::CommandletOptions const& CommandOptions = GenerateBlueprintAPIUtils::CommandOptions;
 
+	FModuleManager::Get().LoadModuleChecked(TEXT("Blutility"));
+
 	FString ActiveFilePath;
 	FArchive* FileOut = nullptr;
 
diff --git a/Engine/Source/Editor/UnrealEd/Private/EditorEngine.cpp b/Engine/Source/Editor/UnrealEd/Private/EditorEngine.cpp
--- a/Engine/Source/Editor/UnrealEd/Private/EditorEngine.cpp
+++ b/Engine/Source/Editor/UnrealEd/Private/EditorEngine.cpp
@@ -685,6 +685,7 @@ void UEditorEngine::Init(IEngineLoop* InEngineLoop)
 		if (!IsRunningCommandlet())
 		{
 			FModuleManager::Get().LoadModule(TEXT("EditorLiveStreaming"));
+			FModuleManager::Get().LoadModule(TEXT("IntroTutorials"));
 		}
 
 		if( FParse::Param( FCommandLine::Get(),TEXT( "PListEditor" ) ) )
diff --git a/Engine/Source/Runtime/Launch/Private/LaunchEngineLoop.cpp b/Engine/Source/Runtime/Launch/Private/LaunchEngineLoop.cpp
--- a/Engine/Source/Runtime/Launch/Private/LaunchEngineLoop.cpp
+++ b/Engine/Source/Runtime/Launch/Private/LaunchEngineLoop.cpp
@@ -1931,9 +1931,11 @@ bool FEngineLoop::LoadStartupCoreModules()
 	// We need this for blueprint projects that have online functionality.
 	FModuleManager::Get().LoadModule(TEXT("OnlineBlueprintSupport"));
 
-	FModuleManager::Get().LoadModule(TEXT("IntroTutorials"));
-	FModuleManager::Get().LoadModule(TEXT("Blutility"));
-
+	if (IsRunningCommandlet())
+	{
+		FModuleManager::Get().LoadModule(TEXT("IntroTutorials"));
+		FModuleManager::Get().LoadModule(TEXT("Blutility"));
+	}
 #endif //(WITH_EDITOR && !(UE_BUILD_SHIPPING || UE_BUILD_TEST))
 
 	return bSuccess;
