Make sure games on Mac save files in user dir, not inside the app bundle

diff --git a/Engine/Source/Runtime/Launch/Private/Mac/LaunchMac.cpp b/Engine/Source/Runtime/Launch/Private/Mac/LaunchMac.cpp
--- a/Engine/Source/Runtime/Launch/Private/Mac/LaunchMac.cpp
+++ b/Engine/Source/Runtime/Launch/Private/Mac/LaunchMac.cpp
@@ -275,6 +275,11 @@ INT32_MAIN_INT32_ARGC_TCHAR_ARGV()
 		GSavedCommandLine += Argument;
 	}
 
+#if UE_GAME
+	// On Mac we always want games to save files to user dir instead of inside the app bundle
+	GSavedCommandLine += TEXT(" -installed");
+#endif
+
 	SCOPED_AUTORELEASE_POOL;
 	[NSApplication sharedApplication];
 	[NSApp setDelegate:[UE4AppDelegate new]];
