Fixed: Device orientation doesn't work properly on Android UE-18157
Fixed: Context loss can cause back buffer to resize on Android UE-17512

diff --git a/Engine/Source/Runtime/Core/Private/Android/AndroidApplication.cpp b/Engine/Source/Runtime/Core/Private/Android/AndroidApplication.cpp
--- a/Engine/Source/Runtime/Core/Private/Android/AndroidApplication.cpp
+++ b/Engine/Source/Runtime/Core/Private/Android/AndroidApplication.cpp
@@ -7,6 +7,8 @@
 
 DEFINE_LOG_CATEGORY_STATIC(LogAndroidApplication, Log, All);
 
+bool FAndroidApplication::bWindowSizeChanged = false;
+
 FAndroidApplication* FAndroidApplication::CreateAndroidApplication()
 {
 	return new FAndroidApplication();
@@ -29,6 +31,24 @@ void FAndroidApplication::PollGameDeviceState( const float TimeDelta )
 	// Poll game device state and send new events
 	InputInterface->Tick( TimeDelta );
 	InputInterface->SendControllerEvents();
+	
+	if (bWindowSizeChanged && 
+		Windows.Num() > 0 && 
+		FPlatformMisc::GetHardwareWindow() != nullptr)
+	{
+		FAndroidWindow::InvalidateCachedScreenRect();
+		FAndroidAppEntry::ReInitWindow();
+				
+		int32 WindowX,WindowY, WindowWidth,WindowHeight;
+		Windows[0]->GetFullScreenInfo(WindowX, WindowY, WindowWidth, WindowHeight);
+		GenericApplication::GetMessageHandler()->OnSizeChanged(Windows[0],WindowWidth,WindowHeight, false);
+		GenericApplication::GetMessageHandler()->OnResizingWindow(Windows[0]);
+		
+		FDisplayMetrics DisplayMetrics;
+		FDisplayMetrics::GetDisplayMetrics(DisplayMetrics);
+		BroadcastDisplayMetricsChanged(DisplayMetrics);
+		bWindowSizeChanged = false;
+	}
 }
 
 FPlatformRect FAndroidApplication::GetWorkArea( const FPlatformRect& CurrentWindow ) const
@@ -79,6 +99,11 @@ void FAndroidApplication::InitializeWindow( const TSharedRef< FGenericWindow >&
 	Window->Initialize( this, InDefinition, ParentWindow, bShowImmediately );
 }
 
+void FAndroidApplication::OnWindowSizeChanged()
+{
+	bWindowSizeChanged = true;
+}
+
 //////////////////////////////////////////////////////////////////////////
 // FJNIHelper
 static JavaVM* CurrentJavaVM = nullptr;
diff --git a/Engine/Source/Runtime/Core/Private/Android/AndroidWindow.cpp b/Engine/Source/Runtime/Core/Private/Android/AndroidWindow.cpp
--- a/Engine/Source/Runtime/Core/Private/Android/AndroidWindow.cpp
+++ b/Engine/Source/Runtime/Core/Private/Android/AndroidWindow.cpp
@@ -72,6 +72,11 @@ int32 FAndroidWindow::GetDepthBufferPreference()
 	return GAndroidDepthBufferPreference;
 }
 
+void FAndroidWindow::InvalidateCachedScreenRect()
+{
+	WindowInit = false;
+}
+
 FPlatformRect FAndroidWindow::GetScreenRect()
 {
 	// CSF is a multiplier to 1280x720
@@ -209,3 +214,12 @@ void FAndroidWindow::CalculateSurfaceSize(void* InWindow, int32_t& SurfaceWidth,
 	SurfaceWidth  = ((SurfaceWidth  + DividableBy - 1) / DividableBy) * DividableBy;
 	SurfaceHeight = ((SurfaceHeight + DividableBy - 1) / DividableBy) * DividableBy;
 }
+
+void FAndroidWindow::OnWindowOrientationChanged(bool bIsPortrait)
+{
+	if (GAndroidIsPortrait != bIsPortrait)
+	{
+		UE_LOG(LogAndroid, Log, TEXT("Window orientation changed: %s"), bIsPortrait ? TEXT("Portrait") : TEXT("Landscape"));
+		GAndroidIsPortrait = bIsPortrait;
+	}
+}
diff --git a/Engine/Source/Runtime/Core/Public/Android/AndroidApplication.h b/Engine/Source/Runtime/Core/Public/Android/AndroidApplication.h
--- a/Engine/Source/Runtime/Core/Public/Android/AndroidApplication.h
+++ b/Engine/Source/Runtime/Core/Public/Android/AndroidApplication.h
@@ -64,6 +64,8 @@ public:
 	
 	void InitializeWindow( const TSharedRef< FGenericWindow >& InWindow, const TSharedRef< FGenericWindowDefinition >& InDefinition, const TSharedPtr< FGenericWindow >& InParent, const bool bShowImmediately );
 
+	static void OnWindowSizeChanged();
+
 private:
 
 	FAndroidApplication();
@@ -74,6 +76,8 @@ private:
 	TSharedPtr< class FAndroidInputInterface > InputInterface;
 
 	TArray< TSharedRef< FAndroidWindow > > Windows;
+
+	static bool bWindowSizeChanged;
 };
 
 PRAGMA_ENABLE_DEPRECATION_WARNINGS
\ No newline at end of file
diff --git a/Engine/Source/Runtime/Core/Public/Android/AndroidWindow.h b/Engine/Source/Runtime/Core/Public/Android/AndroidWindow.h
--- a/Engine/Source/Runtime/Core/Public/Android/AndroidWindow.h
+++ b/Engine/Source/Runtime/Core/Public/Android/AndroidWindow.h
@@ -36,8 +36,10 @@ public:
 	virtual void SetOSWindowHandle(void*);
 
 	static FPlatformRect GetScreenRect();
+	static void InvalidateCachedScreenRect();
 
 	static void CalculateSurfaceSize(void* InWindow, int32_t& SurfaceWidth, int32_t& SurfaceHeight);
+	static void OnWindowOrientationChanged(bool bIsPortrait);
 
 	static int32 GetDepthBufferPreference();
 
diff --git a/Engine/Source/Runtime/Launch/Private/Android/AndroidEventManager.cpp b/Engine/Source/Runtime/Launch/Private/Android/AndroidEventManager.cpp
--- a/Engine/Source/Runtime/Launch/Private/Android/AndroidEventManager.cpp
+++ b/Engine/Source/Runtime/Launch/Private/Android/AndroidEventManager.cpp
@@ -42,10 +42,10 @@ void FAppEventManager::Tick()
 			PendingWindow = (ANativeWindow*)Event.Data;
 			break;
 		case APP_EVENT_STATE_WINDOW_RESIZED:
-			bWindowChanged = true;
+			FAndroidApplication::OnWindowSizeChanged();
 			break;
 		case APP_EVENT_STATE_WINDOW_CHANGED:
-			bWindowChanged = true;
+			FAndroidApplication::OnWindowSizeChanged();
 			break;
 		case APP_EVENT_STATE_SAVE_STATE:
 			bSaveState = true; //todo android: handle save state.
@@ -111,20 +111,6 @@ void FAppEventManager::Tick()
 			}
 		}
 
-		if (bWindowChanged)
-		{
-			// wait until window is valid and created. 
-			if(FPlatformMisc::GetHardwareWindow() != NULL)
-			{
-				if(GEngine && GEngine->GameViewport && GEngine->GameViewport->ViewportFrame)
-				{
-					ExecWindowChanged();
-					bWindowChanged = false;
-					bHaveWindow = true;
-				}
-			}
-		}
-
 		if (!bRunning && bHaveWindow && bHaveGame)
 		{
 			ResumeRendering();
@@ -186,7 +172,6 @@ void FAppEventManager::Tick()
 FAppEventManager::FAppEventManager():
 	FirstInitialized(false)
 	,bCreateWindow(false)
-	,bWindowChanged(false)
 	,bWindowInFocus(false)
 	,bSaveState(false)
 	,bAudioPaused(false)
@@ -281,18 +266,6 @@ void FAppEventManager::ExecWindowCreated()
 }
 
 
-void FAppEventManager::ExecWindowChanged()
-{
-	FPlatformRect ScreenRect = FAndroidWindow::GetScreenRect();
-	UE_LOG(LogAndroidEvents, Display, TEXT("ExecWindowChanged : width: %u, height:  %u"), ScreenRect.Right, ScreenRect.Bottom);
-
-	if(GEngine && GEngine->GameViewport && GEngine->GameViewport->ViewportFrame)
-	{
-		GEngine->GameViewport->ViewportFrame->ResizeFrame(ScreenRect.Right, ScreenRect.Bottom, EWindowMode::Fullscreen);
-	}
-}
-
-
 void FAppEventManager::PauseAudio()
 {
 	bAudioPaused = true;
diff --git a/Engine/Source/Runtime/Launch/Private/Android/AndroidEventManager.h b/Engine/Source/Runtime/Launch/Private/Android/AndroidEventManager.h
--- a/Engine/Source/Runtime/Launch/Private/Android/AndroidEventManager.h
+++ b/Engine/Source/Runtime/Launch/Private/Android/AndroidEventManager.h
@@ -58,7 +58,6 @@ private:
 	void PauseAudio();
 	void ResumeAudio();
 
-	void ExecWindowChanged();
 	void ExecWindowCreated();
 	
 	static FAppEventManager* sInstance;
@@ -75,7 +74,6 @@ private:
 	bool bCreateWindow;
 	ANativeWindow* PendingWindow;
 
-	bool bWindowChanged;
 	bool bWindowInFocus;
 	bool bSaveState;
 	bool bAudioPaused;
diff --git a/Engine/Source/Runtime/Launch/Private/Android/LaunchAndroid.cpp b/Engine/Source/Runtime/Launch/Private/Android/LaunchAndroid.cpp
--- a/Engine/Source/Runtime/Launch/Private/Android/LaunchAndroid.cpp
+++ b/Engine/Source/Runtime/Launch/Private/Android/LaunchAndroid.cpp
@@ -650,7 +650,10 @@ static void onNativeWindowResized(ANativeActivity* activity, ANativeWindow* wind
 	FAppEventManager::GetInstance()->EnqueueAppEvent(APP_EVENT_STATE_WINDOW_CHANGED,window);
 }
 
-
+static void onContentRectChanged(ANativeActivity *activity, const ARect *rect)
+{
+	FAppEventManager::GetInstance()->EnqueueAppEvent(APP_EVENT_STATE_WINDOW_CHANGED, nullptr);
+}
 
 //Called from the event process thread
 static void OnAppCommandCB(struct android_app* app, int32_t cmd)
@@ -760,10 +763,14 @@ static void OnAppCommandCB(struct android_app* app, int32_t cmd)
 		UE_LOG(LogAndroid, Log, TEXT("Case APP_CMD_CONTENT_RECT_CHANGED"));
 		break;
 	case APP_CMD_CONFIG_CHANGED:
-		/**
-		 * Command from main thread: the current device configuration has changed.
-		 */
-		UE_LOG(LogAndroid, Log, TEXT("Case APP_CMD_CONFIG_CHANGED"));
+		{
+			/**
+			* Command from main thread: the current device configuration has changed.
+			*/
+			UE_LOG(LogAndroid, Log, TEXT("Case APP_CMD_CONFIG_CHANGED"));
+			bool bPortrait = (AConfiguration_getOrientation(app->config) == ACONFIGURATION_ORIENTATION_PORT);
+			FAndroidWindow::OnWindowOrientationChanged(bPortrait);
+		}
 		break;
 	case APP_CMD_LOW_MEMORY:
 		/**
@@ -777,7 +784,8 @@ static void OnAppCommandCB(struct android_app* app, int32_t cmd)
 		 * Command from main thread: the app's activity has been started.
 		 */
 		UE_LOG(LogAndroid, Log, TEXT("Case APP_CMD_START"));
-		app->activity->callbacks->onNativeWindowResized  = onNativeWindowResized; //currently not handled in glue code.
+		app->activity->callbacks->onNativeWindowResized = onNativeWindowResized; //currently not handled in glue code.
+		app->activity->callbacks->onContentRectChanged = onContentRectChanged; 
 		FAppEventManager::GetInstance()->EnqueueAppEvent(APP_EVENT_STATE_ON_START);
 	
 		break;
