UECORE-234/Enable stats for commandlets

diff --git a/Engine/Source/Runtime/Core/Private/Stats/Stats2.cpp b/Engine/Source/Runtime/Core/Private/Stats/Stats2.cpp
--- a/Engine/Source/Runtime/Core/Private/Stats/Stats2.cpp
+++ b/Engine/Source/Runtime/Core/Private/Stats/Stats2.cpp
@@ -195,6 +195,24 @@ void FStats::AdvanceFrame( bool bDiscardCallstack, const FOnAdvanceRenderingThre
 #endif
 }
 
+void FStats::TickCommandletStats()
+{
+	if (HasStatsForCommandletsToken())
+	{
+		FThreadStats* ThreadStats = FThreadStats::GetThreadStats();
+		//check( ThreadStats->ScopeCount == 0 && TEXT( "FStats::TickCommandletStats must be called outside any scope counters" ) );
+
+		FTaskGraphInterface::Get().ProcessThreadUntilIdle( ENamedThreads::GameThread );
+		FTicker::GetCoreTicker().Tick( 1 / 60.0f );
+
+		FStats::AdvanceFrame( false );
+	}
+}
+
+bool FStats::HasStatsForCommandletsToken()
+{
+	return FParse::Param( FCommandLine::Get(), TEXT( "StatsForCommandlets" ) );
+}
 
 /* Todo
 
@@ -1000,7 +1018,7 @@ FThreadStats::FThreadStats( EConstructor ):
 void FThreadStats::CheckEnable()
 {
 	bool bOldMasterEnable(bMasterEnable);
-	bool bNewMasterEnable( WillEverCollectData() && !IsRunningCommandlet() && IsThreadingReady() && (MasterEnableCounter.GetValue()) );
+	bool bNewMasterEnable( WillEverCollectData() && (!IsRunningCommandlet() || FStats::HasStatsForCommandletsToken()) && IsThreadingReady() && (MasterEnableCounter.GetValue()) );
 	if (bMasterEnable != bNewMasterEnable)
 	{
 		MasterDisableChangeTagLockAdd();
@@ -1228,7 +1246,7 @@ void FThreadStats::StartThread()
 
 	CheckForCollectingStartupStats();
 
-	UE_LOG( LogStats, Log, TEXT( "Stats thread started" ) );
+	UE_LOG( LogStats, Log, TEXT( "Stats thread started at %f" ), FPlatformTime::Seconds() - GStartTime );
 }
 
 static FGraphEventRef LastFramesEvents[MAX_STAT_LAG];
diff --git a/Engine/Source/Runtime/Core/Public/Stats/Stats2.h b/Engine/Source/Runtime/Core/Public/Stats/Stats2.h
--- a/Engine/Source/Runtime/Core/Public/Stats/Stats2.h
+++ b/Engine/Source/Runtime/Core/Public/Stats/Stats2.h
@@ -61,6 +61,15 @@ struct CORE_API FStats
 	/** Advances stats for the current frame. */
 	static void AdvanceFrame( bool bDiscardCallstack, const FOnAdvanceRenderingThreadStats& AdvanceRenderingThreadStatsDelegate = FOnAdvanceRenderingThreadStats() );
 
+	/** Advances stats for commandlets, only valid if the command line has the proper token. @see HasStatsForCommandletsToken */
+	static void TickCommandletStats();
+
+	/**
+	* @return true, if the command line has the StatsForCommandlets token which enables stats in the commandlets.
+	* !!!CAUTION!!! You need to manually advanced frame in order to maintain the data integrity and not to leak the memory.
+	*/
+	static bool HasStatsForCommandletsToken();
+
 	/** Current game thread stats frame. */
 	static int32 GameThreadStatsFrame;
 };
@@ -1527,6 +1536,15 @@ public:
 			}
 		}
 	}
+
+	/**
+	 * Stops the capturing and stores the result and resets the stat id.
+	 */
+	FORCEINLINE_STATS void StopAndResetStatId()
+	{
+		Stop();
+		StatId = NAME_None;
+	}
 };
 
 /** Manages startup messages, usually to update the metadata. */
diff --git a/Engine/Source/Runtime/Launch/Private/LaunchEngineLoop.cpp b/Engine/Source/Runtime/Launch/Private/LaunchEngineLoop.cpp
--- a/Engine/Source/Runtime/Launch/Private/LaunchEngineLoop.cpp
+++ b/Engine/Source/Runtime/Launch/Private/LaunchEngineLoop.cpp
@@ -681,10 +681,10 @@ bool IsServerDelegateForOSS(FName WorldContextHandle)
 }
 #endif
 
+DECLARE_CYCLE_STAT( TEXT( "FEngineLoop::PreInit.AfterStats" ), STAT_FEngineLoop_PreInit_AfterStats, STATGROUP_LoadTime );
+
 int32 FEngineLoop::PreInit( const TCHAR* CmdLine )
 {
-	DECLARE_SCOPE_CYCLE_COUNTER(TEXT("Engine Pre-Initialized"), STAT_PreInit, STATGROUP_LoadTime);
-
 	if (FParse::Param(CmdLine, TEXT("UTF8Output")))
 	{
 		FPlatformMisc::SetUTF8Output();
@@ -925,7 +925,11 @@ int32 FEngineLoop::PreInit( const TCHAR* CmdLine )
 		if (bHasCommandletToken)
 		{
 #if STATS
-			FThreadStats::MasterDisableForever();
+			// Leave the stats enabled.
+			if (!FStats::HasStatsForCommandletsToken())
+			{
+				FThreadStats::MasterDisableForever();
+			}
 #endif
 			if (Token.StartsWith(TEXT("run=")))
 			{
@@ -982,7 +986,11 @@ int32 FEngineLoop::PreInit( const TCHAR* CmdLine )
 	if (bHasCommandletToken)
 	{
 #if STATS
-		FThreadStats::MasterDisableForever();
+		// Leave the stats enabled.
+		if (!FStats::HasStatsForCommandletsToken())
+		{
+			FThreadStats::MasterDisableForever();
+		}
 #endif
 		if (Token.StartsWith(TEXT("run=")))
 		{
@@ -1057,7 +1065,7 @@ int32 FEngineLoop::PreInit( const TCHAR* CmdLine )
 	FThreadStats::StartThread();
 #endif
 
-	DECLARE_SCOPE_CYCLE_COUNTER( TEXT( "FEngineLoop::PreInit.AfterStats" ), STAT_FEngineLoop_PreInit_AfterStats, STATGROUP_LoadTime );
+	FScopeCycleCounter CycleCount_AfterStats( GET_STATID( STAT_FEngineLoop_PreInit_AfterStats ) );
 
 	// Load Core modules required for everything else to work (needs to be loaded before InitializeRenderingCVarsCaching)
 	LoadCoreModules();
@@ -1601,8 +1609,12 @@ int32 FEngineLoop::PreInit( const TCHAR* CmdLine )
 
 			// Commandlets don't always handle -run= properly in the commandline so we'll provide them
 			// with a custom version that doesn't have it.
-			Commandlet->ParseParms( CommandletCommandLine );
+			Commandlet->ParseParms( CommandletCommandLine ); 
+			// We have to close the scope, otherwise we will end with broken stats.
+			CycleCount_AfterStats.StopAndResetStatId();
+			FStats::TickCommandletStats();
 			int32 ErrorLevel = Commandlet->Main( CommandletCommandLine );
+			FStats::TickCommandletStats();
 
 			// Log warning/ error summary.
 			if( Commandlet->ShowErrorCount )
@@ -2532,7 +2544,7 @@ void FEngineLoop::Tick()
 		const FCPUTime CPUTime = FPlatformTime::GetCPUTime();
 		SET_FLOAT_STAT( STAT_CPUTimePct, CPUTime.CPUTimePct );
 		SET_FLOAT_STAT( STAT_CPUTimePctRelative, CPUTime.CPUTimePctRelative );
-	} 
+	}
 }
 
 void FEngineLoop::ClearPendingCleanupObjects()
