Compilation fix

diff --git a/Engine/Source/Runtime/Launch/Private/LaunchEngineLoop.cpp b/Engine/Source/Runtime/Launch/Private/LaunchEngineLoop.cpp
--- a/Engine/Source/Runtime/Launch/Private/LaunchEngineLoop.cpp
+++ b/Engine/Source/Runtime/Launch/Private/LaunchEngineLoop.cpp
@@ -1610,8 +1610,10 @@ int32 FEngineLoop::PreInit( const TCHAR* CmdLine )
 			// Commandlets don't always handle -run= properly in the commandline so we'll provide them
 			// with a custom version that doesn't have it.
 			Commandlet->ParseParms( CommandletCommandLine ); 
+#if	STATS
 			// We have to close the scope, otherwise we will end with broken stats.
 			CycleCount_AfterStats.StopAndResetStatId();
+#endif // STATS
 			FStats::TickCommandletStats();
 			int32 ErrorLevel = Commandlet->Main( CommandletCommandLine );
 			FStats::TickCommandletStats();
