Pass correct event for onDestroy
#ue4
#android

diff --git a/Engine/Source/Runtime/Launch/Private/Android/LaunchAndroid.cpp b/Engine/Source/Runtime/Launch/Private/Android/LaunchAndroid.cpp
--- a/Engine/Source/Runtime/Launch/Private/Android/LaunchAndroid.cpp
+++ b/Engine/Source/Runtime/Launch/Private/Android/LaunchAndroid.cpp
@@ -819,7 +819,7 @@ static void OnAppCommandCB(struct android_app* app, int32_t cmd)
 		* and waiting for the app thread to clean up and exit before proceeding.
 		*/
 		UE_LOG(LogAndroid, Log, TEXT("Case APP_CMD_DESTROY"));
-		FAppEventManager::GetInstance()->EnqueueAppEvent(APP_EVENT_STATE_ON_PAUSE);
+		FAppEventManager::GetInstance()->EnqueueAppEvent(APP_EVENT_STATE_ON_DESTROY);
 		break;
 	}
 
