Fixed: Resolution incorrect when launching from the editor on Nexus 5 with the device locked
UE-18636

diff --git a/Engine/Source/Runtime/Launch/Private/Android/AndroidEventManager.cpp b/Engine/Source/Runtime/Launch/Private/Android/AndroidEventManager.cpp
--- a/Engine/Source/Runtime/Launch/Private/Android/AndroidEventManager.cpp
+++ b/Engine/Source/Runtime/Launch/Private/Android/AndroidEventManager.cpp
@@ -261,6 +261,11 @@ void FAppEventManager::ExecWindowCreated()
 
 	FPlatformMisc::SetHardwareWindow(PendingWindow);
 	PendingWindow = NULL;
+
+	// When application launched while device is in sleep mode SystemResolution could be set to opposite orientation values
+	// Force to update SystemResolution to current values whenever we create a new window
+	FPlatformRect ScreenRect = FAndroidWindow::GetScreenRect();
+	FSystemResolution::RequestResolutionChange(ScreenRect.Right, ScreenRect.Bottom, EWindowMode::Fullscreen);
 	
 	FAndroidAppEntry::ReInitWindow();
 }
diff --git a/Engine/Source/Runtime/Launch/Private/Android/LaunchAndroid.cpp b/Engine/Source/Runtime/Launch/Private/Android/LaunchAndroid.cpp
--- a/Engine/Source/Runtime/Launch/Private/Android/LaunchAndroid.cpp
+++ b/Engine/Source/Runtime/Launch/Private/Android/LaunchAndroid.cpp
@@ -647,11 +647,13 @@ static int32_t HandleInputCB(struct android_app* app, AInputEvent* event)
 
 static void onNativeWindowResized(ANativeActivity* activity, ANativeWindow* window)
 {
-	FAppEventManager::GetInstance()->EnqueueAppEvent(APP_EVENT_STATE_WINDOW_CHANGED,window);
+	UE_LOG(LogAndroid, Log, TEXT("On Native Window Resized: %p"), window);
+	FAppEventManager::GetInstance()->EnqueueAppEvent(APP_EVENT_STATE_WINDOW_CHANGED, window);
 }
 
 static void onContentRectChanged(ANativeActivity *activity, const ARect *rect)
 {
+	UE_LOG(LogAndroid, Log, TEXT("On Content Rect Changed: left: %d, top: %d, right: %d, bottom: %d"), rect->left, rect->top, rect->right, rect->bottom);
 	FAppEventManager::GetInstance()->EnqueueAppEvent(APP_EVENT_STATE_WINDOW_CHANGED, nullptr);
 }
 
