Fix Android CIS (fallout from keyboard changes).

diff --git a/Engine/Source/Runtime/Launch/Private/Android/LaunchAndroid.cpp b/Engine/Source/Runtime/Launch/Private/Android/LaunchAndroid.cpp
--- a/Engine/Source/Runtime/Launch/Private/Android/LaunchAndroid.cpp
+++ b/Engine/Source/Runtime/Launch/Private/Android/LaunchAndroid.cpp
@@ -44,17 +44,17 @@ static const int32_t AxisList[] =
 };
 
 // map of all supported keycodes
-static TSet<uint16> MappedKeyCodes;
+static TSet<uint32> MappedKeyCodes;
 
 // List of gamepad keycodes to ignore
-static const uint16 IgnoredGamepadKeyCodesList[] =
+static const uint32 IgnoredGamepadKeyCodesList[] =
 {
 	AKEYCODE_VOLUME_UP,
 	AKEYCODE_VOLUME_DOWN
 };
 
 // List of desired gamepad keycodes
-static const uint16 ValidGamepadKeyCodesList[] =
+static const uint32 ValidGamepadKeyCodesList[] =
 {
 	AKEYCODE_BUTTON_A,
 	AKEYCODE_DPAD_CENTER,
@@ -78,10 +78,10 @@ static const uint16 ValidGamepadKeyCodesList[] =
 };
 
 // map of gamepad keycodes that should be ignored
-static TSet<uint16> IgnoredGamepadKeyCodes;
+static TSet<uint32> IgnoredGamepadKeyCodes;
 
 // map of gamepad keycodes that should be passed forward
-static TSet<uint16> ValidGamepadKeyCodes;
+static TSet<uint32> ValidGamepadKeyCodes;
 
 // -nostdlib means no crtbegin_so.o, so we have to provide our own __dso_handle and atexit()
 extern "C"
@@ -252,7 +252,7 @@ int32 AndroidMain(struct android_app* state)
 
 	// setup key filtering
 	static const uint32 MAX_KEY_MAPPINGS(256);
-	uint16 KeyCodes[MAX_KEY_MAPPINGS];
+	uint32 KeyCodes[MAX_KEY_MAPPINGS];
 	uint32 NumKeyCodes = FPlatformMisc::GetKeyMap(KeyCodes, nullptr, MAX_KEY_MAPPINGS);
 
 	for (int i = 0; i < NumKeyCodes; ++i)
@@ -260,13 +260,13 @@ int32 AndroidMain(struct android_app* state)
 		MappedKeyCodes.Add(KeyCodes[i]);
 	}
 
-	const int IgnoredGamepadKeyCodeCount = sizeof(IgnoredGamepadKeyCodesList)/sizeof(uint16);
+	const int IgnoredGamepadKeyCodeCount = sizeof(IgnoredGamepadKeyCodesList)/sizeof(uint32);
 	for (int i = 0; i < IgnoredGamepadKeyCodeCount; ++i)
 	{
 		IgnoredGamepadKeyCodes.Add(IgnoredGamepadKeyCodesList[i]);
 	}
 
-	const int ValidGamepadKeyCodeCount = sizeof(ValidGamepadKeyCodesList)/sizeof(uint16);
+	const int ValidGamepadKeyCodeCount = sizeof(ValidGamepadKeyCodesList)/sizeof(uint32);
 	for (int i = 0; i < ValidGamepadKeyCodeCount; ++i)
 	{
 		ValidGamepadKeyCodes.Add(ValidGamepadKeyCodesList[i]);
