Update version in main to 4.10.

diff --git a/Engine/Build/Build.version b/Engine/Build/Build.version
--- a/Engine/Build/Build.version
+++ b/Engine/Build/Build.version
@@ -1,6 +1,6 @@
 {
 	"MajorVersion": 4,
-	"MinorVersion": 9,
+	"MinorVersion": 10,
 	"PatchVersion": 0,
 	"Changelist": 0,
 	"IsLicenseeVersion": 0,
diff --git a/Engine/Source/Runtime/Launch/Resources/Version.h b/Engine/Source/Runtime/Launch/Resources/Version.h
--- a/Engine/Source/Runtime/Launch/Resources/Version.h
+++ b/Engine/Source/Runtime/Launch/Resources/Version.h
@@ -3,7 +3,7 @@
 #pragma once
 
 #define ENGINE_MAJOR_VERSION	4
-#define ENGINE_MINOR_VERSION	9
+#define ENGINE_MINOR_VERSION	10
 #define ENGINE_PATCH_VERSION	0
 
 #define ENGINE_IS_LICENSEE_VERSION 0
