Fixed: Android: Games crash upon returning to home during inital load
UE-18705

diff --git a/Engine/Source/Runtime/Core/Private/Android/AndroidWindow.cpp b/Engine/Source/Runtime/Core/Private/Android/AndroidWindow.cpp
--- a/Engine/Source/Runtime/Core/Private/Android/AndroidWindow.cpp
+++ b/Engine/Source/Runtime/Core/Private/Android/AndroidWindow.cpp
@@ -77,6 +77,16 @@ void FAndroidWindow::InvalidateCachedScreenRect()
 	WindowInit = false;
 }
 
+void FAndroidWindow::AcquireWindowRef(ANativeWindow* InWindow)
+{
+	ANativeWindow_acquire(InWindow);
+}
+
+void FAndroidWindow::ReleaseWindowRef(ANativeWindow* InWindow)
+{
+	ANativeWindow_release(InWindow);
+}
+
 FPlatformRect FAndroidWindow::GetScreenRect()
 {
 	// CSF is a multiplier to 1280x720
diff --git a/Engine/Source/Runtime/Core/Public/Android/AndroidWindow.h b/Engine/Source/Runtime/Core/Public/Android/AndroidWindow.h
--- a/Engine/Source/Runtime/Core/Public/Android/AndroidWindow.h
+++ b/Engine/Source/Runtime/Core/Public/Android/AndroidWindow.h
@@ -42,6 +42,9 @@ public:
 	static void OnWindowOrientationChanged(bool bIsPortrait);
 
 	static int32 GetDepthBufferPreference();
+	
+	static void AcquireWindowRef(ANativeWindow* InWindow);
+	static void ReleaseWindowRef(ANativeWindow* InWindow);
 
 protected:
 	/** @return true if the native window is currently in fullscreen mode, false otherwise */
diff --git a/Engine/Source/Runtime/Launch/Private/Android/AndroidEventManager.cpp b/Engine/Source/Runtime/Launch/Private/Android/AndroidEventManager.cpp
--- a/Engine/Source/Runtime/Launch/Private/Android/AndroidEventManager.cpp
+++ b/Engine/Source/Runtime/Launch/Private/Android/AndroidEventManager.cpp
@@ -193,6 +193,10 @@ void FAppEventManager::HandleWindowCreated(void* InWindow)
 	rc = pthread_mutex_unlock(&MainMutex);
 	check(rc == 0);
 
+	// Make sure window will not be deleted until event is processed
+	// Window could be deleted by OS while event queue stuck at game start-up phase
+	FAndroidWindow::AcquireWindowRef((ANativeWindow*)InWindow);
+
 	if(AlreadyInited)
 	{
 		EnqueueAppEvent(APP_EVENT_STATE_WINDOW_CREATED, InWindow );
@@ -260,14 +264,18 @@ void FAppEventManager::ExecWindowCreated()
 	check(PendingWindow)
 
 	FPlatformMisc::SetHardwareWindow(PendingWindow);
-	PendingWindow = NULL;
-
+	
 	// When application launched while device is in sleep mode SystemResolution could be set to opposite orientation values
 	// Force to update SystemResolution to current values whenever we create a new window
 	FPlatformRect ScreenRect = FAndroidWindow::GetScreenRect();
 	FSystemResolution::RequestResolutionChange(ScreenRect.Right, ScreenRect.Bottom, EWindowMode::Fullscreen);
 	
 	FAndroidAppEntry::ReInitWindow();
+	
+	// We hold this reference to ensure that window will not be deleted while game starting up
+	// release it when window is finally initialized
+	FAndroidWindow::ReleaseWindowRef(PendingWindow);
+	PendingWindow = nullptr;
 }
 
 
