UE4 - WinXP compile fix

diff --git a/Engine/Source/Runtime/Launch/Launch.Build.cs b/Engine/Source/Runtime/Launch/Launch.Build.cs
--- a/Engine/Source/Runtime/Launch/Launch.Build.cs
+++ b/Engine/Source/Runtime/Launch/Launch.Build.cs
@@ -50,8 +50,11 @@ public class Launch : ModuleRules
 			if ((Target.Platform == UnrealTargetPlatform.Win32) ||
 				(Target.Platform == UnrealTargetPlatform.Win64))
 			{
+				if (!WindowsPlatform.IsWindowsXPSupported())
+				{
+					DynamicallyLoadedModuleNames.Add("D3D12RHI");
+				}
 				DynamicallyLoadedModuleNames.Add("D3D11RHI");
-				DynamicallyLoadedModuleNames.Add("D3D12RHI");
 				DynamicallyLoadedModuleNames.Add("XAudio2");
 			}
 			else if (Target.Platform == UnrealTargetPlatform.Mac)
