Temporaly disabling FVersionedModuleEnumerator as it does not discovers modules in sub-folders

diff --git a/Engine/Source/Runtime/Launch/Private/LaunchEngineLoop.cpp b/Engine/Source/Runtime/Launch/Private/LaunchEngineLoop.cpp
--- a/Engine/Source/Runtime/Launch/Private/LaunchEngineLoop.cpp
+++ b/Engine/Source/Runtime/Launch/Private/LaunchEngineLoop.cpp
@@ -1031,7 +1031,7 @@ int32 FEngineLoop::PreInit( const TCHAR* CmdLine )
 	}
 
 	// Set up the module list and version information, if it's not compiled-in
-#if !IS_MONOLITHIC && BUILT_FROM_CHANGELIST == 0
+#if 0 && !IS_MONOLITHIC && BUILT_FROM_CHANGELIST == 0
 	static FVersionedModuleEnumerator ModuleEnumerator;
 	if(ModuleEnumerator.RegisterWithModuleManager())
 	{
