- Fixed init order to setup streaming after UObjects are initialized (texture streaming uses DeviceProfiles which are initialized via properties with the UObject system).

diff --git a/Engine/Source/Runtime/Launch/Private/LaunchEngineLoop.cpp b/Engine/Source/Runtime/Launch/Private/LaunchEngineLoop.cpp
--- a/Engine/Source/Runtime/Launch/Private/LaunchEngineLoop.cpp
+++ b/Engine/Source/Runtime/Launch/Private/LaunchEngineLoop.cpp
@@ -1359,9 +1359,6 @@ int32 FEngineLoop::PreInit( const TCHAR* CmdLine )
 		GDistanceFieldAsyncQueue = new FDistanceFieldAsyncQueue();
 	}
 
-	// Initialize the texture streaming system (needs to happen after RHIInit).
-	IStreamingManager::Get();
-
 	{
 		DECLARE_SCOPE_CYCLE_COUNTER(TEXT("Initial UObject load"), STAT_InitialUObjectLoad, STATGROUP_LoadTime);
 
@@ -1393,6 +1390,9 @@ int32 FEngineLoop::PreInit( const TCHAR* CmdLine )
 		UMaterialInterface::AssertDefaultMaterialsPostLoaded();
 	}
 
+	// Initialize the texture streaming system (needs to happen after RHIInit and ProcessNewlyLoadedUObjects).
+	IStreamingManager::Get();
+
 	SlowTask.EnterProgressFrame(5);
 
 	// Tell the module manager is may now process newly-loaded UObjects when new C++ modules are loaded
