Fixed: Graphical corruption on Android when switching back to app
UE-19160
UE-19254

diff --git a/Engine/Source/Runtime/Core/Private/Android/AndroidApplication.cpp b/Engine/Source/Runtime/Core/Private/Android/AndroidApplication.cpp
--- a/Engine/Source/Runtime/Core/Private/Android/AndroidApplication.cpp
+++ b/Engine/Source/Runtime/Core/Private/Android/AndroidApplication.cpp
@@ -36,11 +36,9 @@ void FAndroidApplication::PollGameDeviceState( const float TimeDelta )
 		Windows.Num() > 0 && 
 		FPlatformMisc::GetHardwareWindow() != nullptr)
 	{
-		FAndroidWindow::InvalidateCachedScreenRect();
-		FAndroidAppEntry::ReInitWindow();
-				
 		int32 WindowX,WindowY, WindowWidth,WindowHeight;
 		Windows[0]->GetFullScreenInfo(WindowX, WindowY, WindowWidth, WindowHeight);
+
 		GenericApplication::GetMessageHandler()->OnSizeChanged(Windows[0],WindowWidth,WindowHeight, false);
 		GenericApplication::GetMessageHandler()->OnResizingWindow(Windows[0]);
 		
diff --git a/Engine/Source/Runtime/Core/Private/Android/AndroidWindow.cpp b/Engine/Source/Runtime/Core/Private/Android/AndroidWindow.cpp
--- a/Engine/Source/Runtime/Core/Private/Android/AndroidWindow.cpp
+++ b/Engine/Source/Runtime/Core/Private/Android/AndroidWindow.cpp
@@ -218,7 +218,8 @@ void FAndroidWindow::CalculateSurfaceSize(void* InWindow, int32_t& SurfaceWidth,
 	SurfaceHeight = ANativeWindow_getHeight(Window);
 
 	// some phones gave it the other way (so, if swap if the app is landscape, but width < height)
-	if (!GAndroidIsPortrait && SurfaceWidth < SurfaceHeight)
+	if ((GAndroidIsPortrait && SurfaceWidth > SurfaceHeight) || 
+		(!GAndroidIsPortrait && SurfaceWidth < SurfaceHeight))
 	{
 		Swap(SurfaceWidth, SurfaceHeight);
 	}
diff --git a/Engine/Source/Runtime/Launch/Private/Android/AndroidEventManager.cpp b/Engine/Source/Runtime/Launch/Private/Android/AndroidEventManager.cpp
--- a/Engine/Source/Runtime/Launch/Private/Android/AndroidEventManager.cpp
+++ b/Engine/Source/Runtime/Launch/Private/Android/AndroidEventManager.cpp
@@ -28,6 +28,8 @@ FAppEventManager* FAppEventManager::GetInstance()
 
 void FAppEventManager::Tick()
 {
+	bool bWindowCreatedThisTick = false;
+	
 	while (!Queue.IsEmpty())
 	{
 		bool bDestroyWindow = false;
@@ -41,11 +43,15 @@ void FAppEventManager::Tick()
 			bCreateWindow = true;
 			PendingWindow = (ANativeWindow*)Event.Data;
 			break;
+		
 		case APP_EVENT_STATE_WINDOW_RESIZED:
-			FAndroidApplication::OnWindowSizeChanged();
-			break;
 		case APP_EVENT_STATE_WINDOW_CHANGED:
-			FAndroidApplication::OnWindowSizeChanged();
+			// React on device orientation/windowSize changes only when application has window
+			// In case window was created this tick it should already has correct size
+			if (bHaveWindow && !bWindowCreatedThisTick)
+			{
+				ExecWindowResized();
+			}
 			break;
 		case APP_EVENT_STATE_SAVE_STATE:
 			bSaveState = true; //todo android: handle save state.
@@ -108,6 +114,7 @@ void FAppEventManager::Tick()
 				ExecWindowCreated();
 				bCreateWindow = false;
 				bHaveWindow = true;
+				bWindowCreatedThisTick = true;
 			}
 		}
 
@@ -276,8 +283,20 @@ void FAppEventManager::ExecWindowCreated()
 	// release it when window is finally initialized
 	FAndroidWindow::ReleaseWindowRef(PendingWindow);
 	PendingWindow = nullptr;
+
+	FAndroidApplication::OnWindowSizeChanged();
 }
 
+void FAppEventManager::ExecWindowResized()
+{
+	if (bRunning)
+	{
+		FlushRenderingCommands();
+	}
+	FAndroidWindow::InvalidateCachedScreenRect();
+	FAndroidAppEntry::ReInitWindow();
+	FAndroidApplication::OnWindowSizeChanged();
+}
 
 void FAppEventManager::PauseAudio()
 {
diff --git a/Engine/Source/Runtime/Launch/Private/Android/AndroidEventManager.h b/Engine/Source/Runtime/Launch/Private/Android/AndroidEventManager.h
--- a/Engine/Source/Runtime/Launch/Private/Android/AndroidEventManager.h
+++ b/Engine/Source/Runtime/Launch/Private/Android/AndroidEventManager.h
@@ -59,6 +59,7 @@ private:
 	void ResumeAudio();
 
 	void ExecWindowCreated();
+	void ExecWindowResized();
 	
 	static FAppEventManager* sInstance;
 
diff --git a/Engine/Source/Runtime/Launch/Private/Android/LaunchAndroid.cpp b/Engine/Source/Runtime/Launch/Private/Android/LaunchAndroid.cpp
--- a/Engine/Source/Runtime/Launch/Private/Android/LaunchAndroid.cpp
+++ b/Engine/Source/Runtime/Launch/Private/Android/LaunchAndroid.cpp
@@ -645,18 +645,6 @@ static int32_t HandleInputCB(struct android_app* app, AInputEvent* event)
 	return 0;
 }
 
-static void onNativeWindowResized(ANativeActivity* activity, ANativeWindow* window)
-{
-	UE_LOG(LogAndroid, Log, TEXT("On Native Window Resized: %p"), window);
-	FAppEventManager::GetInstance()->EnqueueAppEvent(APP_EVENT_STATE_WINDOW_CHANGED, window);
-}
-
-static void onContentRectChanged(ANativeActivity *activity, const ARect *rect)
-{
-	UE_LOG(LogAndroid, Log, TEXT("On Content Rect Changed: left: %d, top: %d, right: %d, bottom: %d"), rect->left, rect->top, rect->right, rect->bottom);
-	FAppEventManager::GetInstance()->EnqueueAppEvent(APP_EVENT_STATE_WINDOW_CHANGED, nullptr);
-}
-
 //Called from the event process thread
 static void OnAppCommandCB(struct android_app* app, int32_t cmd)
 {
@@ -770,8 +758,11 @@ static void OnAppCommandCB(struct android_app* app, int32_t cmd)
 			* Command from main thread: the current device configuration has changed.
 			*/
 			UE_LOG(LogAndroid, Log, TEXT("Case APP_CMD_CONFIG_CHANGED"));
+			
 			bool bPortrait = (AConfiguration_getOrientation(app->config) == ACONFIGURATION_ORIENTATION_PORT);
 			FAndroidWindow::OnWindowOrientationChanged(bPortrait);
+
+			FAppEventManager::GetInstance()->EnqueueAppEvent(APP_EVENT_STATE_WINDOW_CHANGED, nullptr);
 		}
 		break;
 	case APP_CMD_LOW_MEMORY:
@@ -786,8 +777,6 @@ static void OnAppCommandCB(struct android_app* app, int32_t cmd)
 		 * Command from main thread: the app's activity has been started.
 		 */
 		UE_LOG(LogAndroid, Log, TEXT("Case APP_CMD_START"));
-		app->activity->callbacks->onNativeWindowResized = onNativeWindowResized; //currently not handled in glue code.
-		app->activity->callbacks->onContentRectChanged = onContentRectChanged; 
 		FAppEventManager::GetInstance()->EnqueueAppEvent(APP_EVENT_STATE_ON_START);
 	
 		break;
diff --git a/Engine/Source/Runtime/OpenGLDrv/Private/Android/AndroidEGL.cpp b/Engine/Source/Runtime/OpenGLDrv/Private/Android/AndroidEGL.cpp
--- a/Engine/Source/Runtime/OpenGLDrv/Private/Android/AndroidEGL.cpp
+++ b/Engine/Source/Runtime/OpenGLDrv/Private/Android/AndroidEGL.cpp
@@ -503,6 +503,7 @@ void AndroidEGL::InitSurface(bool bUseSmallSurface)
 
 void AndroidEGL::ReInit()
 {
+	SetCurrentContext(EGL_NO_CONTEXT, EGL_NO_SURFACE);
 	InitSurface(false);
 	SetCurrentSharedContext();
 }
