UE-17753 - fixed issue with message dialog hang
#ios

diff --git a/Engine/Source/Runtime/Launch/Private/IOS/IOSAppDelegateConsoleHandling.cpp b/Engine/Source/Runtime/Launch/Private/IOS/IOSAppDelegateConsoleHandling.cpp
--- a/Engine/Source/Runtime/Launch/Private/IOS/IOSAppDelegateConsoleHandling.cpp
+++ b/Engine/Source/Runtime/Launch/Private/IOS/IOSAppDelegateConsoleHandling.cpp
@@ -30,6 +30,7 @@ extern bool GShowSplashScreen;
 									style:UIAlertActionStyleDefault
 									handler:^(UIAlertAction* action)
 									{
+										self.AlertResponse = 1;
 										[self.ConsoleAlertController dismissViewControllerAnimated : YES completion : nil];
 
 										// we clicked Ok (not Cancel at index 0), submit the console command
@@ -42,6 +43,7 @@ extern bool GShowSplashScreen;
 										style : UIAlertActionStyleDefault
 										handler : ^ (UIAlertAction* action)
 										{
+											self.AlertResponse = 0;
 											[self.ConsoleAlertController dismissViewControllerAnimated : YES completion : nil];
 										}
 		];
@@ -173,6 +175,8 @@ extern bool GShowSplashScreen;
 											style : UIAlertActionStyleDefault
 											handler : ^ (UIAlertAction* action)
 											{
+												// just set our AlertResponse property, all we need to do
+												self.AlertResponse = OptionalButtonIndex - 2;
 												[AlertController dismissViewControllerAnimated : YES completion : nil];
 											}
 			];
