PR #1455: Don't show GUI message on project compilation failure in unattended mode (Contributed by [REDACTED_PERSON])

diff --git a/Engine/Source/Runtime/Launch/Private/LaunchEngineLoop.cpp b/Engine/Source/Runtime/Launch/Private/LaunchEngineLoop.cpp
--- a/Engine/Source/Runtime/Launch/Private/LaunchEngineLoop.cpp
+++ b/Engine/Source/Runtime/Launch/Private/LaunchEngineLoop.cpp
@@ -2751,7 +2751,10 @@ bool FEngineLoop::AppInit( )
 					{
 						UE_LOG(LogInit, Warning, TEXT("Still incompatible or missing module: %s"), *StillIncompatibleFiles[Idx]);
 					}
-					FPlatformMisc::MessageBoxExt(EAppMsgType::Ok, *FString::Printf(TEXT("%s could not be compiled. Try rebuilding from source manually."), FApp::GetGameName()), TEXT("Error"));
+					if (!FApp::IsUnattended())
+					{
+						FPlatformMisc::MessageBoxExt(EAppMsgType::Ok, *FString::Printf(TEXT("%s could not be compiled. Try rebuilding from source manually."), FApp::GetGameName()), TEXT("Error"));
+					}
 					return false;
 				}
 			}
