UE4 - Fix RHIThread not getting deleted when quitting

diff --git a/Engine/Source/Runtime/Launch/Private/LaunchEngineLoop.cpp b/Engine/Source/Runtime/Launch/Private/LaunchEngineLoop.cpp
--- a/Engine/Source/Runtime/Launch/Private/LaunchEngineLoop.cpp
+++ b/Engine/Source/Runtime/Launch/Private/LaunchEngineLoop.cpp
@@ -177,6 +177,20 @@ static TScopedPointer<FOutputDeviceConsole>	GScopedLogConsole;
 static TScopedPointer<FOutputDeviceStdOutput> GScopedStdOut;
 static TScopedPointer<FOutputDeviceTestExit> GScopedTestExit;
 
+
+static void RHIExitAndStopRHIThread()
+{
+	RHIExit();
+
+	// Stop the RHI Thread
+	if (GUseRHIThread)
+	{
+		DECLARE_CYCLE_STAT(TEXT("Wait For RHIThread Finish"), STAT_WaitForRHIThreadFinish, STATGROUP_TaskGraphTasks);
+		FGraphEventRef QuitTask = TGraphTask<FReturnGraphTask>::CreateTask(nullptr, ENamedThreads::GameThread).ConstructAndDispatchWhenReady(ENamedThreads::RHIThread);
+		FTaskGraphInterface::Get().WaitUntilTaskCompletes(QuitTask, ENamedThreads::GameThread_Local);
+	}
+}
+
 /**
  * Initializes std out device and adds it to GLog
  **/
@@ -2204,7 +2218,7 @@ void FEngineLoop::Exit()
 	StopRenderingThread();
 
 	// Tear down the RHI.
-	RHIExit();
+	RHIExitAndStopRHIThread();
 
 #if WITH_ENGINE
 	// Save the hot reload state
