UE4 - Compile fix

diff --git a/Engine/Source/Runtime/Launch/Private/LaunchEngineLoop.cpp b/Engine/Source/Runtime/Launch/Private/LaunchEngineLoop.cpp
--- a/Engine/Source/Runtime/Launch/Private/LaunchEngineLoop.cpp
+++ b/Engine/Source/Runtime/Launch/Private/LaunchEngineLoop.cpp
@@ -178,6 +178,7 @@ static TScopedPointer<FOutputDeviceStdOutput> GScopedStdOut;
 static TScopedPointer<FOutputDeviceTestExit> GScopedTestExit;
 
 
+#if WITH_ENGINE
 static void RHIExitAndStopRHIThread()
 {
 	RHIExit();
@@ -190,6 +191,7 @@ static void RHIExitAndStopRHIThread()
 		FTaskGraphInterface::Get().WaitUntilTaskCompletes(QuitTask, ENamedThreads::GameThread_Local);
 	}
 }
+#endif
 
 /**
  * Initializes std out device and adds it to GLog
