Merged from Orion CL 2656424

Fix poorly named stat around Slate tick

diff --git a/Engine/Source/Runtime/Launch/Private/LaunchEngineLoop.cpp b/Engine/Source/Runtime/Launch/Private/LaunchEngineLoop.cpp
--- a/Engine/Source/Runtime/Launch/Private/LaunchEngineLoop.cpp
+++ b/Engine/Source/Runtime/Launch/Private/LaunchEngineLoop.cpp
@@ -2463,13 +2463,15 @@ void FEngineLoop::Tick()
 
 		if (FSlateApplication::IsInitialized() && !bIdleMode)
 		{
+			{
 #if WITH_ENGINE
-			QUICK_SCOPE_CYCLE_COUNTER(STAT_FEngineLoop_ProcessPlayerControllersSlateOperations);
+				QUICK_SCOPE_CYCLE_COUNTER(STAT_FEngineLoop_ProcessPlayerControllersSlateOperations);
 #endif
-			check(!IsRunningDedicatedServer());
-			
-			// Process slate operations accumulated in the world ticks.
-			ProcessLocalPlayerSlateOperations();
+				check(!IsRunningDedicatedServer());
+
+				// Process slate operations accumulated in the world ticks.
+				ProcessLocalPlayerSlateOperations();
+			}
 
 			// Tick Slate application
 			FSlateApplication::Get().Tick();
