Add missing dependency for the Linux editor build.

diff --git a/Engine/Source/Runtime/Launch/Launch.Build.cs b/Engine/Source/Runtime/Launch/Launch.Build.cs
--- a/Engine/Source/Runtime/Launch/Launch.Build.cs
+++ b/Engine/Source/Runtime/Launch/Launch.Build.cs
@@ -64,6 +64,7 @@ public class Launch : ModuleRules
 			else if (Target.Platform == UnrealTargetPlatform.Linux)
 			{
 				DynamicallyLoadedModuleNames.Add("ALAudio");
+				PrivateDependencyModuleNames.Add("Json");
 			}
 
 			PrivateIncludePathModuleNames.AddRange(
