Cooker now doesn't load global shaders for the host platform.  Instead only loads global shaders for cooking targets.

diff --git a/Engine/Source/Runtime/Launch/Private/LaunchEngineLoop.cpp b/Engine/Source/Runtime/Launch/Private/LaunchEngineLoop.cpp
--- a/Engine/Source/Runtime/Launch/Private/LaunchEngineLoop.cpp
+++ b/Engine/Source/Runtime/Launch/Private/LaunchEngineLoop.cpp
@@ -1388,10 +1388,22 @@ int32 FEngineLoop::PreInit( const TCHAR* CmdLine )
 		SlowTask.EnterProgressFrame(30);
 	
 		// Load the global shaders.
-		if (GetGlobalShaderMap(GMaxRHIFeatureLevel) == NULL && GIsRequestingExit)
+		// if (!IsRunningCommandlet()) 
+		// hack: don't load global shaders if we are cooking we will load the shaders for the correct platform later
+		FString Commandline = FCommandLine::Get();
+		if (Commandline.Contains(TEXT("cookcommandlet")) == false &&
+			Commandline.Contains(TEXT("run=cook")) == false )
+		// if (FParse::Param(FCommandLine::Get(), TEXT("Multiprocess")) == false)
+		{
+			if (GetGlobalShaderMap(GMaxRHIFeatureLevel) == NULL && GIsRequestingExit)
+			{
+				// This means we can't continue without the global shader map.
+				return 1;
+			}
+		}
+		else if (FPlatformProperties::RequiresCookedData() == false)
 		{
-			// This means we can't continue without the global shader map.
-			return 1;
+			GetDerivedDataCacheRef();
 		}
 
 		// In order to be able to use short script package names get all script
