Correct possible Android crash on entering editable text
#ue4
#android

diff --git a/Engine/Build/Android/Java/src/com/epicgames/ue4/GameActivity.java b/Engine/Build/Android/Java/src/com/epicgames/ue4/GameActivity.java
--- a/Engine/Build/Android/Java/src/com/epicgames/ue4/GameActivity.java
+++ b/Engine/Build/Android/Java/src/com/epicgames/ue4/GameActivity.java
@@ -823,7 +823,8 @@ public class GameActivity extends NativeActivity
 
 		// Set label and starting contents
 		virtualKeyboardAlert.setTitle(Label);
-		virtualKeyboardInputBox.setText(Contents);
+		virtualKeyboardInputBox.setText("");
+		virtualKeyboardInputBox.append(Contents);
 
 		// configure for type of input
 		virtualKeyboardInputBox.setInputType(InputType);
diff --git a/Engine/Source/Runtime/Launch/Private/Android/AndroidJNI.cpp b/Engine/Source/Runtime/Launch/Private/Android/AndroidJNI.cpp
--- a/Engine/Source/Runtime/Launch/Private/Android/AndroidJNI.cpp
+++ b/Engine/Source/Runtime/Launch/Private/Android/AndroidJNI.cpp
@@ -298,9 +298,18 @@ extern "C" void Java_com_epicgames_ue4_GameActivity_nativeVirtualKeyboardResult(
 		if (VirtualKeyboardWidget != NULL)
 		{
 			const char* javaChars = jenv->GetStringUTFChars(contents, 0);
-			VirtualKeyboardWidget->SetTextFromVirtualKeyboard(FText::FromString(FString(UTF8_TO_TCHAR(javaChars))));
 
-			//Release the string
+			// call to set the widget text on game thread
+			if (FTaskGraphInterface::IsRunning())
+			{
+				FGraphEventRef SetWidgetText = FFunctionGraphTask::CreateAndDispatchWhenReady([&]()
+				{
+					VirtualKeyboardWidget->SetTextFromVirtualKeyboard(FText::FromString(FString(UTF8_TO_TCHAR(javaChars))));
+				}, TStatId(), NULL, ENamedThreads::GameThread);
+				FTaskGraphInterface::Get().WaitUntilTaskCompletes(SetWidgetText);
+			}
+
+			// release string
 			jenv->ReleaseStringUTFChars(contents, javaChars);
 		}
 	}
