Engine: Fixed an issue where commandlets did not load plugin modules with a loading phase of PostEngineInit
[INTEGRATE] Change 2683776 by [REDACTED_PERSON][REDACTED_HANDLE] on 2015/09/08 17:22:02

diff --git a/Engine/Source/Runtime/Launch/Private/LaunchEngineLoop.cpp b/Engine/Source/Runtime/Launch/Private/LaunchEngineLoop.cpp
--- a/Engine/Source/Runtime/Launch/Private/LaunchEngineLoop.cpp
+++ b/Engine/Source/Runtime/Launch/Private/LaunchEngineLoop.cpp
@@ -1641,6 +1641,10 @@ int32 FEngineLoop::PreInit( const TCHAR* CmdLine )
 				}
 			}
 
+			// Load all the post-engine init modules
+			ensure(IProjectManager::Get().LoadModulesForProject(ELoadingPhase::PostEngineInit));
+			ensure(IPluginManager::Get().LoadModulesForEnabledPlugins(ELoadingPhase::PostEngineInit));
+
 			//run automation smoke tests now that the commandlet has had a chance to override the above flags and GEngine is available
 #if !PLATFORM_HTML5 && !PLATFORM_HTML5_WIN32 
 			FAutomationTestFramework::GetInstance().RunSmokeTests();
