Do not init HighRes screenshot system on servers.

(Merging CL 2689273 to the main branch).

diff --git a/Engine/Source/Runtime/Launch/Private/LaunchEngineLoop.cpp b/Engine/Source/Runtime/Launch/Private/LaunchEngineLoop.cpp
--- a/Engine/Source/Runtime/Launch/Private/LaunchEngineLoop.cpp
+++ b/Engine/Source/Runtime/Launch/Private/LaunchEngineLoop.cpp
@@ -1827,8 +1827,11 @@ int32 FEngineLoop::PreInit( const TCHAR* CmdLine )
 	}
 #endif
 
-	// Init HighRes screenshot system.
-	GetHighResScreenshotConfig().Init();
+	// Init HighRes screenshot system, unless running on server
+	if (!IsRunningDedicatedServer())
+	{
+		GetHighResScreenshotConfig().Init();
+	}
 
 #else // WITH_ENGINE
 	EndInitTextLocalization();
