Build fix caused from a bad merge

diff --git a/Engine/Source/Runtime/Launch/Private/Android/AndroidJNI.cpp b/Engine/Source/Runtime/Launch/Private/Android/AndroidJNI.cpp
--- a/Engine/Source/Runtime/Launch/Private/Android/AndroidJNI.cpp
+++ b/Engine/Source/Runtime/Launch/Private/Android/AndroidJNI.cpp
@@ -311,8 +311,10 @@ extern "C" void Java_com_epicgames_ue4_GameActivity_nativeVirtualKeyboardResult(
 			{
 				FGraphEventRef SetWidgetText = FFunctionGraphTask::CreateAndDispatchWhenReady([&]()
 				{
-			VirtualKeyboardWidget->SetTextFromVirtualKeyboard(FText::FromString(FString(UTF8_TO_TCHAR(javaChars))),false);
-
+					VirtualKeyboardWidget->SetTextFromVirtualKeyboard(FText::FromString(FString(UTF8_TO_TCHAR(javaChars))), false);
+				}, TStatId(), NULL, ENamedThreads::GameThread);
+				FTaskGraphInterface::Get().WaitUntilTaskCompletes(SetWidgetText);
+			}
 			// release string
 			jenv->ReleaseStringUTFChars(contents, javaChars);
 		}
