Only trigger window resize if orientation changes on APP_CMD_CONFIG_CHANGED
#ue4
#android

diff --git a/Engine/Source/Runtime/Core/Private/Android/AndroidWindow.cpp b/Engine/Source/Runtime/Core/Private/Android/AndroidWindow.cpp
--- a/Engine/Source/Runtime/Core/Private/Android/AndroidWindow.cpp
+++ b/Engine/Source/Runtime/Core/Private/Android/AndroidWindow.cpp
@@ -230,11 +230,13 @@ void FAndroidWindow::CalculateSurfaceSize(void* InWindow, int32_t& SurfaceWidth,
 	SurfaceHeight = ((SurfaceHeight + DividableBy - 1) / DividableBy) * DividableBy;
 }
 
-void FAndroidWindow::OnWindowOrientationChanged(bool bIsPortrait)
+bool FAndroidWindow::OnWindowOrientationChanged(bool bIsPortrait)
 {
 	if (GAndroidIsPortrait != bIsPortrait)
 	{
 		UE_LOG(LogAndroid, Log, TEXT("Window orientation changed: %s"), bIsPortrait ? TEXT("Portrait") : TEXT("Landscape"));
 		GAndroidIsPortrait = bIsPortrait;
+		return true;
 	}
+	return false;
 }
diff --git a/Engine/Source/Runtime/Core/Public/Android/AndroidWindow.h b/Engine/Source/Runtime/Core/Public/Android/AndroidWindow.h
--- a/Engine/Source/Runtime/Core/Public/Android/AndroidWindow.h
+++ b/Engine/Source/Runtime/Core/Public/Android/AndroidWindow.h
@@ -39,7 +39,7 @@ public:
 	static void InvalidateCachedScreenRect();
 
 	static void CalculateSurfaceSize(void* InWindow, int32_t& SurfaceWidth, int32_t& SurfaceHeight);
-	static void OnWindowOrientationChanged(bool bIsPortrait);
+	static bool OnWindowOrientationChanged(bool bIsPortrait);
 
 	static int32 GetDepthBufferPreference();
 	
diff --git a/Engine/Source/Runtime/Launch/Private/Android/LaunchAndroid.cpp b/Engine/Source/Runtime/Launch/Private/Android/LaunchAndroid.cpp
--- a/Engine/Source/Runtime/Launch/Private/Android/LaunchAndroid.cpp
+++ b/Engine/Source/Runtime/Launch/Private/Android/LaunchAndroid.cpp
@@ -754,9 +754,10 @@ static void OnAppCommandCB(struct android_app* app, int32_t cmd)
 			UE_LOG(LogAndroid, Log, TEXT("Case APP_CMD_CONFIG_CHANGED"));
 			
 			bool bPortrait = (AConfiguration_getOrientation(app->config) == ACONFIGURATION_ORIENTATION_PORT);
-			FAndroidWindow::OnWindowOrientationChanged(bPortrait);
-
-			FAppEventManager::GetInstance()->EnqueueAppEvent(APP_EVENT_STATE_WINDOW_CHANGED, nullptr);
+			if (FAndroidWindow::OnWindowOrientationChanged(bPortrait))
+			{
+				FAppEventManager::GetInstance()->EnqueueAppEvent(APP_EVENT_STATE_WINDOW_CHANGED, nullptr);
+			}
 		}
 		break;
 	case APP_CMD_LOW_MEMORY:
