Niagara: Renamed NiagaraAnimation to NiagaraSequence (for consistency with Sequencer name changes)

diff --git a/Engine/Config/BaseEngine.ini b/Engine/Config/BaseEngine.ini
--- a/Engine/Config/BaseEngine.ini
+++ b/Engine/Config/BaseEngine.ini
@@ -578,6 +578,9 @@ bUseStreamingPause=false
 ; Renamed ScriptComponent to ScriptPluginComponent
 +ActiveClassRedirects=(OldClassName="ScriptComponent",NewClassName="ScriptPluginComponent")
 
+; Renamed movie scene animations to movie scene sequences
++ActiveClassRedirects=(OldClassName="NiagaraAnimation",NewClassName="NiagaraSequence")
+
 ; Moved Oculus code to a shared plugin module
 +ActiveGameNameRedirects=(OldGameName="/Script/OculusRift.OculusFunctionLibrary", NewGameName="/Script/OculusLibrary.OculusFunctionLibrary")
 
diff --git a/Engine/Source/Editor/NiagaraEditor/Private/NiagaraEffectEditor.cpp b/Engine/Source/Editor/NiagaraEditor/Private/NiagaraEffectEditor.cpp
--- a/Engine/Source/Editor/NiagaraEditor/Private/NiagaraEffectEditor.cpp
+++ b/Engine/Source/Editor/NiagaraEditor/Private/NiagaraEffectEditor.cpp
@@ -2,7 +2,7 @@
 
 #include "NiagaraEditorPrivatePCH.h"
 #include "NiagaraEffect.h"
-#include "NiagaraAnimation.h"
+#include "NiagaraSequence.h"
 
 #include "Toolkits/IToolkitHost.h"
 #include "Editor/WorkspaceMenuStructure/Public/WorkspaceMenuStructureModule.h"
@@ -86,7 +86,7 @@ void FNiagaraEffectEditor::InitNiagaraEffectEditor(const EToolkitMode::Type Mode
 	if (!Sequencer.IsValid())
 	{
 		MovieScene = NewObject<UMovieScene>(InEffect, FName("Niagara Effect MovieScene"), RF_RootSet);
-		auto NewAnimation = NewObject<UNiagaraAnimation>(MovieScene);
+		auto NewAnimation = NewObject<UNiagaraSequence>(MovieScene);
 		MovieScene->StartTime = InTime;
 		MovieScene->InTime = InTime;
 		MovieScene->OutTime = OutTime;
diff --git a/Engine/Source/Runtime/Launch/Private/LaunchEngineLoop.cpp b/Engine/Source/Runtime/Launch/Private/LaunchEngineLoop.cpp
--- a/Engine/Source/Runtime/Launch/Private/LaunchEngineLoop.cpp
+++ b/Engine/Source/Runtime/Launch/Private/LaunchEngineLoop.cpp
@@ -99,6 +99,7 @@
 #	include "VisualLogger/VisualLogger.h"
 #endif
 
+
 // Pipe output to std output
 // This enables UBT to collect the output for it's own use
 class FOutputDeviceStdOutput : public FOutputDevice
@@ -131,6 +132,7 @@ private:
 	bool bAllowLogVerbosity;
 };
 
+
 // Exits the game/editor if any of the specified phrases appears in the log output
 class FOutputDeviceTestExit : public FOutputDevice
 {
@@ -152,7 +154,7 @@ public:
 				if (FCString::Stristr(V, *Phrase) && !FCString::Stristr(V, TEXT("-testexit=")))
 				{
 #if WITH_ENGINE
-					if (GEngine != NULL)
+					if (GEngine != nullptr)
 					{
 						if (GIsEditor)
 						{
@@ -160,7 +162,7 @@ public:
 						}
 						else
 						{
-							GEngine->Exec(NULL, TEXT("QUIT"));
+							GEngine->Exec(nullptr, TEXT("QUIT"));
 						}
 					}
 #else
@@ -173,6 +175,7 @@ public:
 	}
 };
 
+
 static TScopedPointer<FOutputDeviceConsole>	GScopedLogConsole;
 static TScopedPointer<FOutputDeviceStdOutput> GScopedStdOut;
 static TScopedPointer<FOutputDeviceTestExit> GScopedTestExit;
@@ -193,6 +196,7 @@ static void RHIExitAndStopRHIThread()
 }
 #endif
 
+
 /**
  * Initializes std out device and adds it to GLog
  **/
@@ -205,6 +209,7 @@ void InitializeStdOutDevice()
 	GLog->AddOutputDevice(GScopedStdOut.GetOwnedPointer());
 }
 
+
 bool ParseGameProjectFromCommandLine(const TCHAR* InCmdLine, FString& OutProjectFilePath, FString& OutGameName)
 {
 	const TCHAR *CmdLine = InCmdLine;
@@ -248,6 +253,7 @@ bool ParseGameProjectFromCommandLine(const TCHAR* InCmdLine, FString& OutProject
 	return false;
 }
 
+
 bool LaunchSetGameName(const TCHAR *InCmdLine, FString& OutGameProjectFilePathUnnormalized)
 {
 	if (GIsGameAgnosticExe)
@@ -338,6 +344,7 @@ bool LaunchSetGameName(const TCHAR *InCmdLine, FString& OutGameProjectFilePathUn
 	return true;
 }
 
+
 void LaunchFixGameNameCase()
 {
 #if PLATFORM_DESKTOP && !IS_PROGRAM
@@ -379,7 +386,8 @@ void LaunchFixGameNameCase()
 #endif	//PLATFORM_DESKTOP
 }
 
-static IPlatformFile* ConditionallyCreateFileWrapper(const TCHAR* Name, IPlatformFile* CurrentPlatformFile, const TCHAR* CommandLine, bool* OutFailedToInitialize = NULL, bool* bOutShouldBeUsed = NULL )
+
+static IPlatformFile* ConditionallyCreateFileWrapper(const TCHAR* Name, IPlatformFile* CurrentPlatformFile, const TCHAR* CommandLine, bool* OutFailedToInitialize = nullptr, bool* bOutShouldBeUsed = nullptr )
 {
 	if (OutFailedToInitialize)
 	{
@@ -390,7 +398,7 @@ static IPlatformFile* ConditionallyCreateFileWrapper(const TCHAR* Name, IPlatfor
 		*bOutShouldBeUsed = false;
 	}
 	IPlatformFile* WrapperFile = FPlatformFileManager::Get().GetPlatformFile(Name);
-	if (WrapperFile != NULL && WrapperFile->ShouldBeUsed(CurrentPlatformFile, CommandLine))
+	if (WrapperFile != nullptr && WrapperFile->ShouldBeUsed(CurrentPlatformFile, CommandLine))
 	{
 		if ( bOutShouldBeUsed )
 		{
@@ -403,17 +411,18 @@ static IPlatformFile* ConditionallyCreateFileWrapper(const TCHAR* Name, IPlatfor
 				*OutFailedToInitialize = true;
 			}
 			// Don't delete the platform file. It will be automatically deleted by its module.
-			WrapperFile = NULL;
+			WrapperFile = nullptr;
 		}
 	}
 	else
 	{
 		// Make sure it won't be used.
-		WrapperFile = NULL;
+		WrapperFile = nullptr;
 	}
 	return WrapperFile;
 }
 
+
 /**
  * Look for any file overrides on the command line (i.e. network connection file handler)
  */
@@ -558,6 +567,7 @@ bool LaunchCheckForFileOverride(const TCHAR* CmdLine, bool& OutFileOverrideFound
 	return true;
 }
 
+
 bool LaunchHasIncompleteGameName()
 {
 	if ( FApp::HasGameName() && !FPaths::IsProjectFilePathSet() )
@@ -578,6 +588,7 @@ bool LaunchHasIncompleteGameName()
 	return false;
 }
 
+
 void LaunchUpdateMostRecentProjectFile()
 {
 	// If we are launching without a game name or project file, we should use the last used project file, if it exists
@@ -618,7 +629,7 @@ void LaunchUpdateMostRecentProjectFile()
 
 FEngineLoop::FEngineLoop()
 #if WITH_ENGINE
-	: EngineService(NULL)
+	: EngineService(nullptr)
 #endif
 { }
 
@@ -664,6 +675,7 @@ int32 FEngineLoop::PreInit(int32 ArgC, TCHAR* ArgV[], const TCHAR* AdditionalCom
 	return GEngineLoop.PreInit(*CmdLine);
 }
 
+
 #if WITH_ENGINE
 bool IsServerDelegateForOSS(FName WorldContextHandle)
 {
@@ -672,7 +684,7 @@ bool IsServerDelegateForOSS(FName WorldContextHandle)
 		return true;
 	}
 
-	UWorld* World = NULL;
+	UWorld* World = nullptr;
 	if (WorldContextHandle != NAME_None)
 	{
 		FWorldContext& WorldContext = GEngine->GetWorldContextFromHandleChecked(WorldContextHandle);
@@ -928,7 +940,7 @@ int32 FEngineLoop::PreInit( const TCHAR* CmdLine )
 	const bool bFirstTokenIsGame = (Token == TEXT("-GAME"));
 	const bool bFirstTokenIsServer = (Token == TEXT("-SERVER"));
 	const bool bFirstTokenIsModeOverride = bFirstTokenIsGame || bFirstTokenIsServer || bHasCommandletToken;
-	const TCHAR* CommandletCommandLine = NULL;
+	const TCHAR* CommandletCommandLine = nullptr;
 	if (bFirstTokenIsModeOverride)
 	{
 		bIsNotEditor = true;
@@ -1002,7 +1014,7 @@ int32 FEngineLoop::PreInit( const TCHAR* CmdLine )
 
 	bHasEditorToken = !bIsNotEditor;
 #elif WITH_ENGINE
-	const TCHAR* CommandletCommandLine = NULL;
+	const TCHAR* CommandletCommandLine = nullptr;
 	if (bHasCommandletToken)
 	{
 #if STATS
@@ -1202,7 +1214,7 @@ int32 FEngineLoop::PreInit( const TCHAR* CmdLine )
 	// Let LogConsole know what ini file it should use to save its setting on exit.
 	// We can't use GGameIni inside log console because it's destroyed in the global
 	// scoped pointer and at that moment GGameIni may already be gone.
-	if( GLogConsole != NULL )
+	if( GLogConsole != nullptr )
 	{
 		GLogConsole->SetIniFilename(*GGameIni);
 	}
@@ -1397,7 +1409,7 @@ int32 FEngineLoop::PreInit( const TCHAR* CmdLine )
 			Commandline.Contains(TEXT("run=cook")) == false )
 		// if (FParse::Param(FCommandLine::Get(), TEXT("Multiprocess")) == false)
 		{
-			if (GetGlobalShaderMap(GMaxRHIFeatureLevel) == NULL && GIsRequestingExit)
+			if (GetGlobalShaderMap(GMaxRHIFeatureLevel) == nullptr && GIsRequestingExit)
 			{
 				// This means we can't continue without the global shader map.
 				return 1;
@@ -1539,7 +1551,7 @@ int32 FEngineLoop::PreInit( const TCHAR* CmdLine )
 
 	if (!bHasEditorToken)
 	{
-		UClass* CommandletClass = NULL;
+		UClass* CommandletClass = nullptr;
 
 		if (!bIsRegularClient)
 		{
@@ -1559,7 +1571,7 @@ int32 FEngineLoop::PreInit( const TCHAR* CmdLine )
 			extern bool GIsConsoleExecutable;
 			if (GIsConsoleExecutable)
 			{
-				if (GLogConsole != NULL && GLogConsole->IsAttached())
+				if (GLogConsole != nullptr && GLogConsole->IsAttached())
 				{
 					GLog->RemoveOutputDevice(GLogConsole);
 				}
@@ -1577,7 +1589,7 @@ int32 FEngineLoop::PreInit( const TCHAR* CmdLine )
 			}
 
 			// print output immediately
-			setvbuf(stdout, NULL, _IONBF, 0);
+			setvbuf(stdout, nullptr, _IONBF, 0);
 
 			UE_LOG(LogInit, Log,  TEXT("Executing %s"), *CommandletClass->GetFullName() );
 
@@ -1613,12 +1625,12 @@ int32 FEngineLoop::PreInit( const TCHAR* CmdLine )
 			
 			// allow the commandlet the opportunity to create a custom engine
 			CommandletClass->GetDefaultObject<UCommandlet>()->CreateCustomEngine(CommandletCommandLine);
-			if ( GEngine == NULL )
+			if ( GEngine == nullptr )
 			{
 #if WITH_EDITOR
 				if ( GIsEditor )
 				{
-					UClass* EditorEngineClass = StaticLoadClass( UEditorEngine::StaticClass(), NULL, TEXT("engine-ini:/Script/Engine.Engine.EditorEngine"), NULL, LOAD_None, NULL );
+					UClass* EditorEngineClass = StaticLoadClass( UEditorEngine::StaticClass(), nullptr, TEXT("engine-ini:/Script/Engine.Engine.EditorEngine"), nullptr, LOAD_None, nullptr );
 
 					GEngine = GEditor = NewObject<UEditorEngine>(GetTransientPackage(), EditorEngineClass);
 
@@ -1631,7 +1643,7 @@ int32 FEngineLoop::PreInit( const TCHAR* CmdLine )
 				else
 #endif
 				{
-					UClass* EngineClass = StaticLoadClass( UEngine::StaticClass(), NULL, TEXT("engine-ini:/Script/Engine.Engine.GameEngine"), NULL, LOAD_None, NULL );
+					UClass* EngineClass = StaticLoadClass( UEngine::StaticClass(), nullptr, TEXT("engine-ini:/Script/Engine.Engine.GameEngine"), nullptr, LOAD_None, nullptr );
 
 					// must do this here so that the engine object that we create on the next line receives the correct property values
 					GEngine = NewObject<UEngine>(GetTransientPackage(), EngineClass);
@@ -1778,7 +1790,7 @@ int32 FEngineLoop::PreInit( const TCHAR* CmdLine )
 	// exit if wanted.
 	if( GIsRequestingExit )
 	{
-		if ( GEngine != NULL )
+		if ( GEngine != nullptr )
 		{
 			GEngine->PreExit();
 		}
@@ -1817,7 +1829,7 @@ int32 FEngineLoop::PreInit( const TCHAR* CmdLine )
 	delete [] CommandLineCopy;
 
 	// initialize the pointer, as it is deleted before being assigned in the first frame
-	PendingCleanupObjects = NULL;
+	PendingCleanupObjects = nullptr;
 
 	// Initialize profile visualizers.
 #if !(UE_BUILD_SHIPPING || UE_BUILD_TEST)
@@ -1847,6 +1859,7 @@ int32 FEngineLoop::PreInit( const TCHAR* CmdLine )
 	return 0;
 }
 
+
 void FEngineLoop::LoadCoreModules()
 {
 	// Always attempt to load CoreUObject. It requires additional pre-init which is called from its module's StartupModule method.
@@ -1855,6 +1868,7 @@ void FEngineLoop::LoadCoreModules()
 #endif
 }
 
+
 void FEngineLoop::LoadPreInitModules()
 {	
 	DECLARE_SCOPE_CYCLE_COUNTER(TEXT("Loading PreInit Modules"), STAT_PreInitModules, STATGROUP_LoadTime);
@@ -1894,6 +1908,7 @@ void FEngineLoop::LoadPreInitModules()
 #endif // WITH_ENGINE
 }
 
+
 #if WITH_ENGINE
 
 bool FEngineLoop::LoadStartupCoreModules()
@@ -2020,6 +2035,7 @@ bool FEngineLoop::LoadStartupCoreModules()
 	return bSuccess;
 }
 
+
 bool FEngineLoop::LoadStartupModules()
 {
 	FScopedSlowTask SlowTask(3);
@@ -2091,7 +2107,6 @@ void FEngineLoop::InitTime()
 }
 
 
-
 //called via FCoreDelegates::StarvedGameLoop
 void GameLoopIsStarved()
 {
@@ -2108,18 +2123,18 @@ int32 FEngineLoop::Init()
 	SlowTask.EnterProgressFrame(10);
 
 	// Figure out which UEngine variant to use.
-	UClass* EngineClass = NULL;
+	UClass* EngineClass = nullptr;
 	if( !GIsEditor )
 	{
 		// We're the game.
-		EngineClass = StaticLoadClass( UGameEngine::StaticClass(), NULL, TEXT("engine-ini:/Script/Engine.Engine.GameEngine"), NULL, LOAD_None, NULL );
+		EngineClass = StaticLoadClass( UGameEngine::StaticClass(), nullptr, TEXT("engine-ini:/Script/Engine.Engine.GameEngine"), nullptr, LOAD_None, nullptr );
 		GEngine = NewObject<UEngine>(GetTransientPackage(), EngineClass);
 	}
 	else
 	{
 #if WITH_EDITOR
 		// We're UnrealEd.
-		EngineClass = StaticLoadClass( UUnrealEdEngine::StaticClass(), NULL, TEXT("engine-ini:/Script/Engine.Engine.UnrealEdEngine"), NULL, LOAD_None, NULL );
+		EngineClass = StaticLoadClass( UUnrealEdEngine::StaticClass(), nullptr, TEXT("engine-ini:/Script/Engine.Engine.UnrealEdEngine"), nullptr, LOAD_None, nullptr );
 		GEngine = GEditor = GUnrealEd = NewObject<UUnrealEdEngine>(GetTransientPackage(), EngineClass);
 #else
 		check(0);
@@ -2189,7 +2204,7 @@ void FEngineLoop::Exit()
 	STAT_ADD_CUSTOMMESSAGE_NAME( STAT_NamedMarker, TEXT( "EngineLoop.Exit" ) );
 
 	GIsRunning	= 0;
-	GLogConsole	= NULL;
+	GLogConsole	= nullptr;
 
 	// shutdown visual logger and flush all data
 #if ENABLE_VISUAL_LOG
@@ -2211,7 +2226,7 @@ void FEngineLoop::Exit()
 #if WITH_ENGINE
 	// shut down messaging
 	delete EngineService;
-	EngineService = NULL;
+	EngineService = nullptr;
 
 	if (SessionService.IsValid())
 	{
@@ -2226,14 +2241,14 @@ void FEngineLoop::Exit()
 	}
 #endif // WITH_ENGINE
 
-	MALLOC_PROFILER( GEngine->Exec( NULL, TEXT( "MPROF STOP" ) ); )
+	MALLOC_PROFILER( GEngine->Exec( nullptr, TEXT( "MPROF STOP" ) ); )
 
-	if ( GEngine != NULL )
+	if ( GEngine != nullptr )
 	{
 		GEngine->ShutdownAudioDeviceManager();
 	}
 
-	if ( GEngine != NULL )
+	if ( GEngine != nullptr )
 	{
 		GEngine->PreExit();
 	}
@@ -2282,12 +2297,11 @@ void FEngineLoop::Exit()
 #endif
 
 	FTaskGraphInterface::Shutdown();
-
 	IStreamingManager::Shutdown();
-
 	FIOSystem::Shutdown();
 }
 
+
 void FEngineLoop::ProcessLocalPlayerSlateOperations() const
 {
 	FSlateApplication& SlateApp = FSlateApplication::Get();
@@ -2327,6 +2341,7 @@ void FEngineLoop::ProcessLocalPlayerSlateOperations() const
 	}
 }
 
+
 bool FEngineLoop::ShouldUseIdleMode() const
 {
 	static const auto CVarIdleWhenNotForeground = IConsoleManager::Get().FindTConsoleVariableDataInt(TEXT("t.IdleWhenNotForeground"));
@@ -2353,6 +2368,7 @@ bool FEngineLoop::ShouldUseIdleMode() const
 	return bIdleMode;
 }
 
+
 #if !UE_BUILD_SHIPPING && !UE_BUILD_TEST
 
 #include "StackTracker.h"
@@ -2389,6 +2405,7 @@ static TAutoConsoleVariable<int32> CVarLogGameThreadMallocChurn_StackLen(
 
 extern CORE_API TFunction<void(int32)>* GGameThreadMallocHook;
 
+
 struct FScopedSampleMallocChurn
 {
 	static FStackTracker GGameThreadMallocChurnTracker;
@@ -2470,6 +2487,7 @@ uint64 FScopedSampleMallocChurn::DumpFrame = 0;
 
 #endif
 
+
 void FEngineLoop::Tick()
 {
 #if !UE_BUILD_SHIPPING && !UE_BUILD_TEST
@@ -2546,7 +2564,7 @@ void FEngineLoop::Tick()
 		CalculateFPSTimings();
 
 		// Note the start of a new frame
-		MALLOC_PROFILER(GMalloc->Exec( NULL, *FString::Printf(TEXT("SNAPSHOTMEMORYFRAME")),*GLog));
+		MALLOC_PROFILER(GMalloc->Exec(nullptr, *FString::Printf(TEXT("SNAPSHOTMEMORYFRAME")),*GLog));
 
 		// handle some per-frame tasks on the rendering thread
 		ENQUEUE_UNIQUE_RENDER_COMMAND(
@@ -2745,14 +2763,16 @@ void FEngineLoop::Tick()
 	}
 }
 
+
 void FEngineLoop::ClearPendingCleanupObjects()
 {
 	delete PendingCleanupObjects;
-	PendingCleanupObjects = NULL;
+	PendingCleanupObjects = nullptr;
 }
 
 #endif // WITH_ENGINE
 
+
 static void CheckForPrintTimesOverride()
 {
 	GPrintLogTimes = ELogTimes::None;
@@ -2786,6 +2806,7 @@ static void CheckForPrintTimesOverride()
 	}
 }
 
+
 /* FEngineLoop static interface
  *****************************************************************************/
 
@@ -2805,7 +2826,7 @@ bool FEngineLoop::AppInit( )
 	// Retrieve additional command line arguments from environment variable.
 	FPlatformMisc::GetEnvironmentVariable(TEXT("UE-CmdLineArgs"), CmdLineEnv,ARRAY_COUNT(CmdLineEnv));
 	
-	// Manually NULL terminate just in case. The NULL string is returned above in the error case so
+	// Manually nullptr terminate just in case. The nullptr string is returned above in the error case so
 	// we don't have to worry about that.
 	CmdLineEnv[ARRAY_COUNT(CmdLineEnv)-1] = 0;
 	FString Env = FString(CmdLineEnv).Trim();
@@ -3147,7 +3168,7 @@ void FEngineLoop::AppExit( )
 	{
 		GConfig->Exit();
 		delete GConfig;
-		GConfig = NULL;
+		GConfig = nullptr;
 	}
 
 	if( GLog )
@@ -3158,6 +3179,7 @@ void FEngineLoop::AppExit( )
 	FInternationalization::TearDown();
 }
 
+
 void FEngineLoop::PreInitHMDDevice()
 {
 #if WITH_ENGINE && !UE_SERVER
@@ -3177,4 +3199,5 @@ void FEngineLoop::PreInitHMDDevice()
 #endif // #if WITH_ENGINE && !UE_SERVER
 }
 
+
 #undef LOCTEXT_NAMESPACE
\ No newline at end of file
diff --git a/Engine/Source/Runtime/Niagara/Private/NiagaraAnimation.cpp b/Engine/Source/Runtime/Niagara/Private/NiagaraAnimation.cpp
deleted file mode 100644
diff --git a/Engine/Source/Runtime/Niagara/Private/NiagaraSequence.cpp b/Engine/Source/Runtime/Niagara/Private/NiagaraSequence.cpp
new file mode 100644
--- /dev/null
+++ b/Engine/Source/Runtime/Niagara/Private/NiagaraSequence.cpp
@@ -0,0 +1,80 @@
+// Copyright 1998-2015 Epic Games, Inc. All Rights Reserved.
+
+#include "NiagaraPrivate.h"
+#include "NiagaraSequence.h"
+
+
+/* UNiagaraSequence structors
+ *****************************************************************************/
+
+UNiagaraSequence::UNiagaraSequence(const FObjectInitializer& ObjectInitializer)
+	: Super(ObjectInitializer)
+	, MovieScene(nullptr)
+{ }
+
+
+/* UMovieSceneAnimation overrides
+ *****************************************************************************/
+
+bool UNiagaraSequence::AllowsSpawnableObjects() const
+{
+	return false;
+}
+
+
+void UNiagaraSequence::BindPossessableObject(const FGuid& ObjectId, UObject& PossessedObject)
+{
+}
+
+
+bool UNiagaraSequence::CanPossessObject(UObject& Object) const
+{
+	return false;
+}
+
+
+void UNiagaraSequence::DestroyAllSpawnedObjects()
+{
+}
+
+
+UObject* UNiagaraSequence::FindObject(const FGuid& ObjectId) const
+{
+	return nullptr;
+}
+
+
+FGuid UNiagaraSequence::FindObjectId(UObject& Object) const
+{
+	return FGuid();
+}
+
+
+UMovieScene* UNiagaraSequence::GetMovieScene() const
+{
+	return MovieScene;
+}
+
+
+UObject* UNiagaraSequence::GetParentObject(UObject* Object) const
+{
+	return nullptr;
+}
+
+
+void UNiagaraSequence::SpawnOrDestroyObjects(bool DestroyAll)
+{
+}
+
+
+#if WITH_EDITOR
+bool UNiagaraSequence::TryGetObjectDisplayName(const FGuid& ObjectId, FText& OutDisplayName) const
+{
+	return false;
+}
+#endif
+
+
+void UNiagaraSequence::UnbindPossessableObjects(const FGuid& ObjectId)
+{
+}
diff --git a/Engine/Source/Runtime/Niagara/Public/NiagaraAnimation.h b/Engine/Source/Runtime/Niagara/Public/NiagaraAnimation.h
deleted file mode 100644
diff --git a/Engine/Source/Runtime/Niagara/Public/NiagaraSequence.h b/Engine/Source/Runtime/Niagara/Public/NiagaraSequence.h
new file mode 100644
--- /dev/null
+++ b/Engine/Source/Runtime/Niagara/Public/NiagaraSequence.h
@@ -0,0 +1,40 @@
+// Copyright 1998-2015 Epic Games, Inc. All Rights Reserved.
+
+#pragma once
+
+#include "MovieSceneSequence.h"
+#include "NiagaraSequence.generated.h"
+
+
+/**
+ * Movie scene sequence used by Niagara.
+ */
+UCLASS(BlueprintType, MinimalAPI)
+class UNiagaraSequence
+	: public UMovieSceneSequence
+{
+	GENERATED_UCLASS_BODY()
+
+public:
+
+	// UMovieSceneSequence overrides
+
+	virtual bool AllowsSpawnableObjects() const override;
+	virtual void BindPossessableObject(const FGuid& ObjectId, UObject& PossessedObject) override;
+	virtual bool CanPossessObject(UObject& Object) const override;
+	virtual void DestroyAllSpawnedObjects() override;
+	virtual UObject* FindObject(const FGuid& ObjectId) const override;
+	virtual FGuid FindObjectId(UObject& Object) const override;
+	virtual UMovieScene* GetMovieScene() const override;
+	virtual UObject* GetParentObject(UObject* Object) const override;
+	virtual void SpawnOrDestroyObjects(bool DestroyAll) override;
+	virtual void UnbindPossessableObjects(const FGuid& ObjectId) override;
+
+#if WITH_EDITOR
+	virtual bool TryGetObjectDisplayName(const FGuid& ObjectId, FText& OutDisplayName) const override;
+#endif
+
+	/** Pointer to the movie scene that controls this sequence. */
+	UPROPERTY()
+	UMovieScene* MovieScene;
+};
