Back out revision 721 from //depot/UE4/Engine/Source/Runtime/Launch/Private/LaunchEngineLoop.cpp

(slipped into the wrong changelist)

diff --git a/Engine/Source/Runtime/Launch/Private/LaunchEngineLoop.cpp b/Engine/Source/Runtime/Launch/Private/LaunchEngineLoop.cpp
--- a/Engine/Source/Runtime/Launch/Private/LaunchEngineLoop.cpp
+++ b/Engine/Source/Runtime/Launch/Private/LaunchEngineLoop.cpp
@@ -99,7 +99,6 @@
 #	include "VisualLogger/VisualLogger.h"
 #endif
 
-
 // Pipe output to std output
 // This enables UBT to collect the output for it's own use
 class FOutputDeviceStdOutput : public FOutputDevice
@@ -132,7 +131,6 @@ private:
 	bool bAllowLogVerbosity;
 };
 
-
 // Exits the game/editor if any of the specified phrases appears in the log output
 class FOutputDeviceTestExit : public FOutputDevice
 {
@@ -154,7 +152,7 @@ public:
 				if (FCString::Stristr(V, *Phrase) && !FCString::Stristr(V, TEXT("-testexit=")))
 				{
 #if WITH_ENGINE
-					if (GEngine != nullptr)
+					if (GEngine != NULL)
 					{
 						if (GIsEditor)
 						{
@@ -162,7 +160,7 @@ public:
 						}
 						else
 						{
-							GEngine->Exec(nullptr, TEXT("QUIT"));
+							GEngine->Exec(NULL, TEXT("QUIT"));
 						}
 					}
 #else
@@ -175,7 +173,6 @@ public:
 	}
 };
 
-
 static TScopedPointer<FOutputDeviceConsole>	GScopedLogConsole;
 static TScopedPointer<FOutputDeviceStdOutput> GScopedStdOut;
 static TScopedPointer<FOutputDeviceTestExit> GScopedTestExit;
@@ -196,7 +193,6 @@ static void RHIExitAndStopRHIThread()
 }
 #endif
 
-
 /**
  * Initializes std out device and adds it to GLog
  **/
@@ -209,7 +205,6 @@ void InitializeStdOutDevice()
 	GLog->AddOutputDevice(GScopedStdOut.GetOwnedPointer());
 }
 
-
 bool ParseGameProjectFromCommandLine(const TCHAR* InCmdLine, FString& OutProjectFilePath, FString& OutGameName)
 {
 	const TCHAR *CmdLine = InCmdLine;
@@ -253,7 +248,6 @@ bool ParseGameProjectFromCommandLine(const TCHAR* InCmdLine, FString& OutProject
 	return false;
 }
 
-
 bool LaunchSetGameName(const TCHAR *InCmdLine, FString& OutGameProjectFilePathUnnormalized)
 {
 	if (GIsGameAgnosticExe)
@@ -344,7 +338,6 @@ bool LaunchSetGameName(const TCHAR *InCmdLine, FString& OutGameProjectFilePathUn
 	return true;
 }
 
-
 void LaunchFixGameNameCase()
 {
 #if PLATFORM_DESKTOP && !IS_PROGRAM
@@ -386,8 +379,7 @@ void LaunchFixGameNameCase()
 #endif	//PLATFORM_DESKTOP
 }
 
-
-static IPlatformFile* ConditionallyCreateFileWrapper(const TCHAR* Name, IPlatformFile* CurrentPlatformFile, const TCHAR* CommandLine, bool* OutFailedToInitialize = nullptr, bool* bOutShouldBeUsed = nullptr )
+static IPlatformFile* ConditionallyCreateFileWrapper(const TCHAR* Name, IPlatformFile* CurrentPlatformFile, const TCHAR* CommandLine, bool* OutFailedToInitialize = NULL, bool* bOutShouldBeUsed = NULL )
 {
 	if (OutFailedToInitialize)
 	{
@@ -398,7 +390,7 @@ static IPlatformFile* ConditionallyCreateFileWrapper(const TCHAR* Name, IPlatfor
 		*bOutShouldBeUsed = false;
 	}
 	IPlatformFile* WrapperFile = FPlatformFileManager::Get().GetPlatformFile(Name);
-	if (WrapperFile != nullptr && WrapperFile->ShouldBeUsed(CurrentPlatformFile, CommandLine))
+	if (WrapperFile != NULL && WrapperFile->ShouldBeUsed(CurrentPlatformFile, CommandLine))
 	{
 		if ( bOutShouldBeUsed )
 		{
@@ -411,18 +403,17 @@ static IPlatformFile* ConditionallyCreateFileWrapper(const TCHAR* Name, IPlatfor
 				*OutFailedToInitialize = true;
 			}
 			// Don't delete the platform file. It will be automatically deleted by its module.
-			WrapperFile = nullptr;
+			WrapperFile = NULL;
 		}
 	}
 	else
 	{
 		// Make sure it won't be used.
-		WrapperFile = nullptr;
+		WrapperFile = NULL;
 	}
 	return WrapperFile;
 }
 
-
 /**
  * Look for any file overrides on the command line (i.e. network connection file handler)
  */
@@ -567,7 +558,6 @@ bool LaunchCheckForFileOverride(const TCHAR* CmdLine, bool& OutFileOverrideFound
 	return true;
 }
 
-
 bool LaunchHasIncompleteGameName()
 {
 	if ( FApp::HasGameName() && !FPaths::IsProjectFilePathSet() )
@@ -588,7 +578,6 @@ bool LaunchHasIncompleteGameName()
 	return false;
 }
 
-
 void LaunchUpdateMostRecentProjectFile()
 {
 	// If we are launching without a game name or project file, we should use the last used project file, if it exists
@@ -629,7 +618,7 @@ void LaunchUpdateMostRecentProjectFile()
 
 FEngineLoop::FEngineLoop()
 #if WITH_ENGINE
-	: EngineService(nullptr)
+	: EngineService(NULL)
 #endif
 { }
 
@@ -675,7 +664,6 @@ int32 FEngineLoop::PreInit(int32 ArgC, TCHAR* ArgV[], const TCHAR* AdditionalCom
 	return GEngineLoop.PreInit(*CmdLine);
 }
 
-
 #if WITH_ENGINE
 bool IsServerDelegateForOSS(FName WorldContextHandle)
 {
@@ -684,7 +672,7 @@ bool IsServerDelegateForOSS(FName WorldContextHandle)
 		return true;
 	}
 
-	UWorld* World = nullptr;
+	UWorld* World = NULL;
 	if (WorldContextHandle != NAME_None)
 	{
 		FWorldContext& WorldContext = GEngine->GetWorldContextFromHandleChecked(WorldContextHandle);
@@ -940,7 +928,7 @@ int32 FEngineLoop::PreInit( const TCHAR* CmdLine )
 	const bool bFirstTokenIsGame = (Token == TEXT("-GAME"));
 	const bool bFirstTokenIsServer = (Token == TEXT("-SERVER"));
 	const bool bFirstTokenIsModeOverride = bFirstTokenIsGame || bFirstTokenIsServer || bHasCommandletToken;
-	const TCHAR* CommandletCommandLine = nullptr;
+	const TCHAR* CommandletCommandLine = NULL;
 	if (bFirstTokenIsModeOverride)
 	{
 		bIsNotEditor = true;
@@ -1014,7 +1002,7 @@ int32 FEngineLoop::PreInit( const TCHAR* CmdLine )
 
 	bHasEditorToken = !bIsNotEditor;
 #elif WITH_ENGINE
-	const TCHAR* CommandletCommandLine = nullptr;
+	const TCHAR* CommandletCommandLine = NULL;
 	if (bHasCommandletToken)
 	{
 #if STATS
@@ -1214,7 +1202,7 @@ int32 FEngineLoop::PreInit( const TCHAR* CmdLine )
 	// Let LogConsole know what ini file it should use to save its setting on exit.
 	// We can't use GGameIni inside log console because it's destroyed in the global
 	// scoped pointer and at that moment GGameIni may already be gone.
-	if( GLogConsole != nullptr )
+	if( GLogConsole != NULL )
 	{
 		GLogConsole->SetIniFilename(*GGameIni);
 	}
@@ -1409,7 +1397,7 @@ int32 FEngineLoop::PreInit( const TCHAR* CmdLine )
 			Commandline.Contains(TEXT("run=cook")) == false )
 		// if (FParse::Param(FCommandLine::Get(), TEXT("Multiprocess")) == false)
 		{
-			if (GetGlobalShaderMap(GMaxRHIFeatureLevel) == nullptr && GIsRequestingExit)
+			if (GetGlobalShaderMap(GMaxRHIFeatureLevel) == NULL && GIsRequestingExit)
 			{
 				// This means we can't continue without the global shader map.
 				return 1;
@@ -1551,7 +1539,7 @@ int32 FEngineLoop::PreInit( const TCHAR* CmdLine )
 
 	if (!bHasEditorToken)
 	{
-		UClass* CommandletClass = nullptr;
+		UClass* CommandletClass = NULL;
 
 		if (!bIsRegularClient)
 		{
@@ -1571,7 +1559,7 @@ int32 FEngineLoop::PreInit( const TCHAR* CmdLine )
 			extern bool GIsConsoleExecutable;
 			if (GIsConsoleExecutable)
 			{
-				if (GLogConsole != nullptr && GLogConsole->IsAttached())
+				if (GLogConsole != NULL && GLogConsole->IsAttached())
 				{
 					GLog->RemoveOutputDevice(GLogConsole);
 				}
@@ -1589,7 +1577,7 @@ int32 FEngineLoop::PreInit( const TCHAR* CmdLine )
 			}
 
 			// print output immediately
-			setvbuf(stdout, nullptr, _IONBF, 0);
+			setvbuf(stdout, NULL, _IONBF, 0);
 
 			UE_LOG(LogInit, Log,  TEXT("Executing %s"), *CommandletClass->GetFullName() );
 
@@ -1625,12 +1613,12 @@ int32 FEngineLoop::PreInit( const TCHAR* CmdLine )
 			
 			// allow the commandlet the opportunity to create a custom engine
 			CommandletClass->GetDefaultObject<UCommandlet>()->CreateCustomEngine(CommandletCommandLine);
-			if ( GEngine == nullptr )
+			if ( GEngine == NULL )
 			{
 #if WITH_EDITOR
 				if ( GIsEditor )
 				{
-					UClass* EditorEngineClass = StaticLoadClass( UEditorEngine::StaticClass(), nullptr, TEXT("engine-ini:/Script/Engine.Engine.EditorEngine"), nullptr, LOAD_None, nullptr );
+					UClass* EditorEngineClass = StaticLoadClass( UEditorEngine::StaticClass(), NULL, TEXT("engine-ini:/Script/Engine.Engine.EditorEngine"), NULL, LOAD_None, NULL );
 
 					GEngine = GEditor = NewObject<UEditorEngine>(GetTransientPackage(), EditorEngineClass);
 
@@ -1643,7 +1631,7 @@ int32 FEngineLoop::PreInit( const TCHAR* CmdLine )
 				else
 #endif
 				{
-					UClass* EngineClass = StaticLoadClass( UEngine::StaticClass(), nullptr, TEXT("engine-ini:/Script/Engine.Engine.GameEngine"), nullptr, LOAD_None, nullptr );
+					UClass* EngineClass = StaticLoadClass( UEngine::StaticClass(), NULL, TEXT("engine-ini:/Script/Engine.Engine.GameEngine"), NULL, LOAD_None, NULL );
 
 					// must do this here so that the engine object that we create on the next line receives the correct property values
 					GEngine = NewObject<UEngine>(GetTransientPackage(), EngineClass);
@@ -1790,7 +1778,7 @@ int32 FEngineLoop::PreInit( const TCHAR* CmdLine )
 	// exit if wanted.
 	if( GIsRequestingExit )
 	{
-		if ( GEngine != nullptr )
+		if ( GEngine != NULL )
 		{
 			GEngine->PreExit();
 		}
@@ -1829,7 +1817,7 @@ int32 FEngineLoop::PreInit( const TCHAR* CmdLine )
 	delete [] CommandLineCopy;
 
 	// initialize the pointer, as it is deleted before being assigned in the first frame
-	PendingCleanupObjects = nullptr;
+	PendingCleanupObjects = NULL;
 
 	// Initialize profile visualizers.
 #if !(UE_BUILD_SHIPPING || UE_BUILD_TEST)
@@ -1859,7 +1847,6 @@ int32 FEngineLoop::PreInit( const TCHAR* CmdLine )
 	return 0;
 }
 
-
 void FEngineLoop::LoadCoreModules()
 {
 	// Always attempt to load CoreUObject. It requires additional pre-init which is called from its module's StartupModule method.
@@ -1868,7 +1855,6 @@ void FEngineLoop::LoadCoreModules()
 #endif
 }
 
-
 void FEngineLoop::LoadPreInitModules()
 {	
 	DECLARE_SCOPE_CYCLE_COUNTER(TEXT("Loading PreInit Modules"), STAT_PreInitModules, STATGROUP_LoadTime);
@@ -1908,7 +1894,6 @@ void FEngineLoop::LoadPreInitModules()
 #endif // WITH_ENGINE
 }
 
-
 #if WITH_ENGINE
 
 bool FEngineLoop::LoadStartupCoreModules()
@@ -2035,7 +2020,6 @@ bool FEngineLoop::LoadStartupCoreModules()
 	return bSuccess;
 }
 
-
 bool FEngineLoop::LoadStartupModules()
 {
 	FScopedSlowTask SlowTask(3);
@@ -2107,6 +2091,7 @@ void FEngineLoop::InitTime()
 }
 
 
+
 //called via FCoreDelegates::StarvedGameLoop
 void GameLoopIsStarved()
 {
@@ -2123,18 +2108,18 @@ int32 FEngineLoop::Init()
 	SlowTask.EnterProgressFrame(10);
 
 	// Figure out which UEngine variant to use.
-	UClass* EngineClass = nullptr;
+	UClass* EngineClass = NULL;
 	if( !GIsEditor )
 	{
 		// We're the game.
-		EngineClass = StaticLoadClass( UGameEngine::StaticClass(), nullptr, TEXT("engine-ini:/Script/Engine.Engine.GameEngine"), nullptr, LOAD_None, nullptr );
+		EngineClass = StaticLoadClass( UGameEngine::StaticClass(), NULL, TEXT("engine-ini:/Script/Engine.Engine.GameEngine"), NULL, LOAD_None, NULL );
 		GEngine = NewObject<UEngine>(GetTransientPackage(), EngineClass);
 	}
 	else
 	{
 #if WITH_EDITOR
 		// We're UnrealEd.
-		EngineClass = StaticLoadClass( UUnrealEdEngine::StaticClass(), nullptr, TEXT("engine-ini:/Script/Engine.Engine.UnrealEdEngine"), nullptr, LOAD_None, nullptr );
+		EngineClass = StaticLoadClass( UUnrealEdEngine::StaticClass(), NULL, TEXT("engine-ini:/Script/Engine.Engine.UnrealEdEngine"), NULL, LOAD_None, NULL );
 		GEngine = GEditor = GUnrealEd = NewObject<UUnrealEdEngine>(GetTransientPackage(), EngineClass);
 #else
 		check(0);
@@ -2204,7 +2189,7 @@ void FEngineLoop::Exit()
 	STAT_ADD_CUSTOMMESSAGE_NAME( STAT_NamedMarker, TEXT( "EngineLoop.Exit" ) );
 
 	GIsRunning	= 0;
-	GLogConsole	= nullptr;
+	GLogConsole	= NULL;
 
 	// shutdown visual logger and flush all data
 #if ENABLE_VISUAL_LOG
@@ -2226,7 +2211,7 @@ void FEngineLoop::Exit()
 #if WITH_ENGINE
 	// shut down messaging
 	delete EngineService;
-	EngineService = nullptr;
+	EngineService = NULL;
 
 	if (SessionService.IsValid())
 	{
@@ -2241,14 +2226,14 @@ void FEngineLoop::Exit()
 	}
 #endif // WITH_ENGINE
 
-	MALLOC_PROFILER( GEngine->Exec( nullptr, TEXT( "MPROF STOP" ) ); )
+	MALLOC_PROFILER( GEngine->Exec( NULL, TEXT( "MPROF STOP" ) ); )
 
-	if ( GEngine != nullptr )
+	if ( GEngine != NULL )
 	{
 		GEngine->ShutdownAudioDeviceManager();
 	}
 
-	if ( GEngine != nullptr )
+	if ( GEngine != NULL )
 	{
 		GEngine->PreExit();
 	}
@@ -2297,11 +2282,12 @@ void FEngineLoop::Exit()
 #endif
 
 	FTaskGraphInterface::Shutdown();
+
 	IStreamingManager::Shutdown();
+
 	FIOSystem::Shutdown();
 }
 
-
 void FEngineLoop::ProcessLocalPlayerSlateOperations() const
 {
 	FSlateApplication& SlateApp = FSlateApplication::Get();
@@ -2341,7 +2327,6 @@ void FEngineLoop::ProcessLocalPlayerSlateOperations() const
 	}
 }
 
-
 bool FEngineLoop::ShouldUseIdleMode() const
 {
 	static const auto CVarIdleWhenNotForeground = IConsoleManager::Get().FindTConsoleVariableDataInt(TEXT("t.IdleWhenNotForeground"));
@@ -2368,7 +2353,6 @@ bool FEngineLoop::ShouldUseIdleMode() const
 	return bIdleMode;
 }
 
-
 #if !UE_BUILD_SHIPPING && !UE_BUILD_TEST
 
 #include "StackTracker.h"
@@ -2405,7 +2389,6 @@ static TAutoConsoleVariable<int32> CVarLogGameThreadMallocChurn_StackLen(
 
 extern CORE_API TFunction<void(int32)>* GGameThreadMallocHook;
 
-
 struct FScopedSampleMallocChurn
 {
 	static FStackTracker GGameThreadMallocChurnTracker;
@@ -2487,7 +2470,6 @@ uint64 FScopedSampleMallocChurn::DumpFrame = 0;
 
 #endif
 
-
 void FEngineLoop::Tick()
 {
 #if !UE_BUILD_SHIPPING && !UE_BUILD_TEST
@@ -2564,7 +2546,7 @@ void FEngineLoop::Tick()
 		CalculateFPSTimings();
 
 		// Note the start of a new frame
-		MALLOC_PROFILER(GMalloc->Exec(nullptr, *FString::Printf(TEXT("SNAPSHOTMEMORYFRAME")),*GLog));
+		MALLOC_PROFILER(GMalloc->Exec( NULL, *FString::Printf(TEXT("SNAPSHOTMEMORYFRAME")),*GLog));
 
 		// handle some per-frame tasks on the rendering thread
 		ENQUEUE_UNIQUE_RENDER_COMMAND(
@@ -2763,16 +2745,14 @@ void FEngineLoop::Tick()
 	}
 }
 
-
 void FEngineLoop::ClearPendingCleanupObjects()
 {
 	delete PendingCleanupObjects;
-	PendingCleanupObjects = nullptr;
+	PendingCleanupObjects = NULL;
 }
 
 #endif // WITH_ENGINE
 
-
 static void CheckForPrintTimesOverride()
 {
 	GPrintLogTimes = ELogTimes::None;
@@ -2806,7 +2786,6 @@ static void CheckForPrintTimesOverride()
 	}
 }
 
-
 /* FEngineLoop static interface
  *****************************************************************************/
 
@@ -2826,7 +2805,7 @@ bool FEngineLoop::AppInit( )
 	// Retrieve additional command line arguments from environment variable.
 	FPlatformMisc::GetEnvironmentVariable(TEXT("UE-CmdLineArgs"), CmdLineEnv,ARRAY_COUNT(CmdLineEnv));
 	
-	// Manually nullptr terminate just in case. The nullptr string is returned above in the error case so
+	// Manually NULL terminate just in case. The NULL string is returned above in the error case so
 	// we don't have to worry about that.
 	CmdLineEnv[ARRAY_COUNT(CmdLineEnv)-1] = 0;
 	FString Env = FString(CmdLineEnv).Trim();
@@ -3168,7 +3147,7 @@ void FEngineLoop::AppExit( )
 	{
 		GConfig->Exit();
 		delete GConfig;
-		GConfig = nullptr;
+		GConfig = NULL;
 	}
 
 	if( GLog )
@@ -3179,7 +3158,6 @@ void FEngineLoop::AppExit( )
 	FInternationalization::TearDown();
 }
 
-
 void FEngineLoop::PreInitHMDDevice()
 {
 #if WITH_ENGINE && !UE_SERVER
@@ -3199,5 +3177,4 @@ void FEngineLoop::PreInitHMDDevice()
 #endif // #if WITH_ENGINE && !UE_SERVER
 }
 
-
 #undef LOCTEXT_NAMESPACE
\ No newline at end of file
