Merging using UE4-To-UE4-Releases-4.10
Send touch events from gamepads to proper controller id
#ue4
#android

diff --git a/Engine/Source/Runtime/Core/Private/Android/AndroidInputInterface.cpp b/Engine/Source/Runtime/Core/Private/Android/AndroidInputInterface.cpp
--- a/Engine/Source/Runtime/Core/Private/Android/AndroidInputInterface.cpp
+++ b/Engine/Source/Runtime/Core/Private/Android/AndroidInputInterface.cpp
@@ -763,19 +763,21 @@ void FAndroidInputInterface::SendControllerEvents()
 	for(int i = 0; i < FAndroidInputInterface::TouchInputStack.Num(); ++i)
 	{
 		TouchInput Touch = FAndroidInputInterface::TouchInputStack[i];
+		int32 ControllerId = FindExistingDevice(Touch.DeviceId);
+		ControllerId = (ControllerId == -1) ? 0 : ControllerId;
 
 		// send input to handler
 		if (Touch.Type == TouchBegan)
 		{
-			MessageHandler->OnTouchStarted( NULL, Touch.Position, Touch.Handle, 0);
+			MessageHandler->OnTouchStarted(NULL, Touch.Position, Touch.Handle, ControllerId);
 		}
 		else if (Touch.Type == TouchEnded)
 		{
-			MessageHandler->OnTouchEnded(Touch.Position, Touch.Handle, 0);
+			MessageHandler->OnTouchEnded(Touch.Position, Touch.Handle, ControllerId);
 		}
 		else
 		{
-			MessageHandler->OnTouchMoved(Touch.Position, Touch.Handle, 0);
+			MessageHandler->OnTouchMoved(Touch.Position, Touch.Handle, ControllerId);
 		}
 	}
 
@@ -923,6 +925,24 @@ void FAndroidInputInterface::QueueTouchInput(TArray<TouchInput> InTouchEvents)
 	FAndroidInputInterface::TouchInputStack.Append(InTouchEvents);
 }
 
+int32 FAndroidInputInterface::FindExistingDevice(int32 deviceId)
+{
+	// Treat non-positive devices ids special
+	if (deviceId < 1)
+		return -1;
+
+	for (int32 ControllerIndex = 0; ControllerIndex < MAX_NUM_CONTROLLERS; ControllerIndex++)
+	{
+		if (DeviceMapping[ControllerIndex].DeviceInfo.DeviceId == deviceId && DeviceMapping[ControllerIndex].DeviceState == MappingState::Valid)
+		{
+			return ControllerIndex;
+		}
+	}
+
+	// Did not find it
+	return -1;
+}
+
 int32 FAndroidInputInterface::GetControllerIndex(int32 deviceId)
 {
 	// Treat non-positive device ids special (always controller 0)
diff --git a/Engine/Source/Runtime/Core/Public/Android/AndroidInputInterface.h b/Engine/Source/Runtime/Core/Public/Android/AndroidInputInterface.h
--- a/Engine/Source/Runtime/Core/Public/Android/AndroidInputInterface.h
+++ b/Engine/Source/Runtime/Core/Public/Android/AndroidInputInterface.h
@@ -116,7 +116,8 @@ struct FAndroidGamepadDeviceMapping
 
 struct TouchInput
 {
-	int Handle;
+	int32 DeviceId;
+	int32 Handle;
 	TouchType Type;
 	FVector2D LastPosition;
 	FVector2D Position;
@@ -225,6 +226,9 @@ private:
 	/** Get controller index corresponding to deviceId (assigns and queries type if new) */
 	static int32 GetControllerIndex(int32 deviceId);
 
+	/** Find controller index corresponding to validated deviceId (returns -1 if not found) */
+	static int32 FindExistingDevice(int32 deviceId);
+
 	/** Push Vibration changes to the controllers */
 	void UpdateVibeMotors();
 
diff --git a/Engine/Source/Runtime/Launch/Private/Android/LaunchAndroid.cpp b/Engine/Source/Runtime/Launch/Private/Android/LaunchAndroid.cpp
--- a/Engine/Source/Runtime/Launch/Private/Android/LaunchAndroid.cpp
+++ b/Engine/Source/Runtime/Launch/Private/Android/LaunchAndroid.cpp
@@ -442,6 +442,7 @@ static int32_t HandleInputCB(struct android_app* app, AInputEvent* event)
 		int actionType = action & AMOTION_EVENT_ACTION_MASK;
 		size_t actionPointer = (size_t)((action & AMOTION_EVENT_ACTION_POINTER_INDEX_MASK) >> AMOTION_EVENT_ACTION_POINTER_INDEX_SHIFT);
 		bool isActionTargeted = (actionType == AMOTION_EVENT_ACTION_POINTER_DOWN || actionType == AMOTION_EVENT_ACTION_POINTER_UP);
+		int32 device = AInputEvent_getDeviceId(event);
 
 		// trap Joystick events first, with fallthrough if there is no joystick support
 		if (((EventSource & AINPUT_SOURCE_CLASS_JOYSTICK) == AINPUT_SOURCE_CLASS_JOYSTICK) &&
@@ -449,7 +450,6 @@ static int32_t HandleInputCB(struct android_app* app, AInputEvent* event)
 			(actionType == AMOTION_EVENT_ACTION_MOVE))
 		{
 			const int axisCount = sizeof(AxisList)/sizeof(int32_t);
-			int32 device = AInputEvent_getDeviceId(event);
 
 			// poll all the axes and forward to update controller state
 			for (int axis = 0; axis < axisCount; axis++)
@@ -535,6 +535,7 @@ static int32_t HandleInputCB(struct android_app* app, AInputEvent* event)
 				UE_LOG(LogAndroid, Verbose, TEXT("Received targeted motion event from pointer %u (id %d) action %d: (%.2f, %.2f)"), actionPointer, pointerId, action, x, y);
 
 				TouchInput TouchMessage;
+				TouchMessage.DeviceId = device;
 				TouchMessage.Handle = pointerId;
 				TouchMessage.Type = type;
 				TouchMessage.Position = FVector2D(x, y);
@@ -555,6 +556,7 @@ static int32_t HandleInputCB(struct android_app* app, AInputEvent* event)
 					UE_LOG(LogAndroid, Verbose, TEXT("Received motion event from index %u (id %d) action %d: (%.2f, %.2f)"), i, pointerId, action, x, y);
 
 					TouchInput TouchMessage;
+					TouchMessage.DeviceId= device;
 					TouchMessage.Handle = pointerId;
 					TouchMessage.Type = type;
 					TouchMessage.Position = FVector2D(x, y);
