Buildmachines will now list ALL unqiue warnings and errors after the commandlet has finished.

Added -DUMPALLWARNINGS command line param to force that behaviour for normal runs.

diff --git a/Engine/Source/Runtime/Launch/Private/LaunchEngineLoop.cpp b/Engine/Source/Runtime/Launch/Private/LaunchEngineLoop.cpp
--- a/Engine/Source/Runtime/Launch/Private/LaunchEngineLoop.cpp
+++ b/Engine/Source/Runtime/Launch/Private/LaunchEngineLoop.cpp
@@ -1691,49 +1691,51 @@ int32 FEngineLoop::PreInit( const TCHAR* CmdLine )
 				{
 					SET_WARN_COLOR(COLOR_WHITE);
 					UE_LOG(LogInit, Display, TEXT(""));
-					UE_LOG(LogInit, Display, TEXT("Warning/Error Summary"));
-					UE_LOG(LogInit, Display, TEXT("---------------------"));
+					UE_LOG(LogInit, Display, TEXT("Warning/Error Summary (Unique only)"));
+					UE_LOG(LogInit, Display, TEXT("-----------------------------------"));
+
+					const int32 MaxMessagesToShow = (GIsBuildMachine || FParse::Param(FCommandLine::Get(), TEXT("DUMPALLWARNINGS"))) ? 
+						FMath::Max(GWarn->Errors.Num(), GWarn->Warnings.Num()) : 50;
 
-					static const int32 MaxMessagesToShow = 50;
 					TSet<FString> ShownMessages;
 					
 					SET_WARN_COLOR(COLOR_RED);
 					ShownMessages.Empty(MaxMessagesToShow);
-					for(auto It = GWarn->Errors.CreateConstIterator(); It; ++It)
+					for (const FString& ErrorMessage : GWarn->Errors)
 					{
 						bool bAlreadyShown = false;
-						ShownMessages.Add(*It, &bAlreadyShown);
+						ShownMessages.Add(ErrorMessage, &bAlreadyShown);
 
 						if (!bAlreadyShown)
 						{
 							if (ShownMessages.Num() > MaxMessagesToShow)
 							{
 								SET_WARN_COLOR(COLOR_WHITE);
-								UE_LOG(LogInit, Display, TEXT("NOTE: Only first %d errors displayed."), MaxMessagesToShow);
+								UE_CLOG(MaxMessagesToShow < GWarn->Errors.Num(), LogInit, Display, TEXT("NOTE: Only first %d errors displayed."), MaxMessagesToShow);
 								break;
 							}
 
-							UE_LOG(LogInit, Display, TEXT("%s"), **It);
+							UE_LOG(LogInit, Display, TEXT("%s"), *ErrorMessage);
 						}
 					}
 
 					SET_WARN_COLOR(COLOR_YELLOW);
 					ShownMessages.Empty(MaxMessagesToShow);
-					for(auto It = GWarn->Warnings.CreateConstIterator(); It; ++It)
+					for (const FString& WarningMessage : GWarn->Warnings)
 					{
 						bool bAlreadyShown = false;
-						ShownMessages.Add(*It, &bAlreadyShown);
+						ShownMessages.Add(WarningMessage, &bAlreadyShown);
 
 						if (!bAlreadyShown)
 						{
 							if (ShownMessages.Num() > MaxMessagesToShow)
 							{
 								SET_WARN_COLOR(COLOR_WHITE);
-								UE_LOG(LogInit, Display, TEXT("NOTE: Only first %d warnings displayed."), MaxMessagesToShow);
+								UE_CLOG(MaxMessagesToShow < GWarn->Warnings.Num(), LogInit, Display, TEXT("NOTE: Only first %d warnings displayed."), MaxMessagesToShow);
 								break;
 							}
 
-							UE_LOG(LogInit, Display, TEXT("%s"), **It);
+							UE_LOG(LogInit, Display, TEXT("%s"), *WarningMessage);
 						}
 					}
 				}
