Fix Version.h after merge

diff --git a/Engine/Source/Runtime/Launch/Resources/Version.h b/Engine/Source/Runtime/Launch/Resources/Version.h
--- a/Engine/Source/Runtime/Launch/Resources/Version.h
+++ b/Engine/Source/Runtime/Launch/Resources/Version.h
@@ -9,7 +9,7 @@
 #define ENGINE_IS_LICENSEE_VERSION 0
 
 #define BUILT_FROM_CHANGELIST 0
-#define BRANCH_NAME "UE4-Orion"
+#define BRANCH_NAME "UE4"
 
 #define ENGINE_IS_PROMOTED_BUILD (BUILT_FROM_CHANGELIST > 0)
 
