Fixed link error on Linux/Mac with VREditor
- Editor no longer references Launch module (UE4Editor).  Linux/Mac doesn't support importing symbols from executable modules.

diff --git a/Engine/Source/Editor/VREditor/VREditor.Build.cs b/Engine/Source/Editor/VREditor/VREditor.Build.cs
--- a/Engine/Source/Editor/VREditor/VREditor.Build.cs
+++ b/Engine/Source/Editor/VREditor/VREditor.Build.cs
@@ -18,8 +18,7 @@ namespace UnrealBuildTool.Rules
 				    "UnrealEd",
 					"UMG",
 					"LevelEditor",
-					"HeadMountedDisplay",
-					"Launch"	// @todo vreditor urgent: Hack needed to access GNewWorldToMetersScale
+					"HeadMountedDisplay"
 				}
 			);
 
diff --git a/Engine/Source/Editor/VREditor/VREditorMode.cpp b/Engine/Source/Editor/VREditor/VREditorMode.cpp
--- a/Engine/Source/Editor/VREditor/VREditorMode.cpp
+++ b/Engine/Source/Editor/VREditor/VREditorMode.cpp
@@ -1766,7 +1766,7 @@ void FVREditorMode::SetWorldToMetersScale( const float NewWorldToMetersScale )
 	// @todo vreditor: This is bad because we're clobbering the world settings which will be saved with the map.  Instead we need to 
 	// be able to apply an override before the scene view gets it
 
-	LAUNCH_API extern float GNewWorldToMetersScale;
+	ENGINE_API extern float GNewWorldToMetersScale;
 	GNewWorldToMetersScale = NewWorldToMetersScale;
 }
 
diff --git a/Engine/Source/Runtime/Engine/Private/WorldSettings.cpp b/Engine/Source/Runtime/Engine/Private/WorldSettings.cpp
--- a/Engine/Source/Runtime/Engine/Private/WorldSettings.cpp
+++ b/Engine/Source/Runtime/Engine/Private/WorldSettings.cpp
@@ -20,6 +20,10 @@
 #include "Editor.h"
 #endif 
 
+// @todo vreditor urgent: Temporary hack to allow world-to-meters to be set before
+// input is polled for motion controller devices each frame.
+ENGINE_API float GNewWorldToMetersScale = 0.0f;
+
 AWorldSettings::AWorldSettings(const FObjectInitializer& ObjectInitializer)
 	: Super(ObjectInitializer.DoNotCreateDefaultSubobject(TEXT("Sprite")))
 {
diff --git a/Engine/Source/Runtime/Launch/Private/LaunchEngineLoop.cpp b/Engine/Source/Runtime/Launch/Private/LaunchEngineLoop.cpp
--- a/Engine/Source/Runtime/Launch/Private/LaunchEngineLoop.cpp
+++ b/Engine/Source/Runtime/Launch/Private/LaunchEngineLoop.cpp
@@ -2593,8 +2593,6 @@ uint64 FScopedSampleMallocChurn::DumpFrame = 0;
 
 #endif
 
-LAUNCH_API float GNewWorldToMetersScale = 0.0f;
-
 void FEngineLoop::Tick()
 {
 #if !UE_BUILD_SHIPPING && !UE_BUILD_TEST
@@ -2704,6 +2702,8 @@ void FEngineLoop::Tick()
 
 		// @todo vreditor urgent: Temporary hack to allow world-to-meters to be set before
 		// input is polled for motion controller devices each frame.
+#if WITH_ENGINE
+		extern ENGINE_API float GNewWorldToMetersScale;
 		if( GNewWorldToMetersScale != 0.0f && GWorld != nullptr )
 		{
 			if( GNewWorldToMetersScale != GWorld->GetWorldSettings()->WorldToMeters )
@@ -2711,6 +2711,7 @@ void FEngineLoop::Tick()
 				GWorld->GetWorldSettings()->WorldToMeters = GNewWorldToMetersScale;
 			}
 		}
+#endif	// WITH_ENGINE
 
 		if (FSlateApplication::IsInitialized() && !bIdleMode)
 		{
