Speculative fix for compile error on mac

diff --git a/Engine/Source/Runtime/Launch/Private/LaunchEngineLoop.cpp b/Engine/Source/Runtime/Launch/Private/LaunchEngineLoop.cpp
--- a/Engine/Source/Runtime/Launch/Private/LaunchEngineLoop.cpp
+++ b/Engine/Source/Runtime/Launch/Private/LaunchEngineLoop.cpp
@@ -20,7 +20,7 @@
 #include "HAL/ThreadHeartBeat.h"
 
 #if WITH_COREUOBJECT
-	#include "CoreUObject/Public/Internationalization/PackageLocalizationManager.h"
+	#include "Internationalization/PackageLocalizationManager.h"
 #endif
 
 #if WITH_EDITOR
