- Better fix for the missing project menus (moving the PacakgeLocCache down after object initialization)

diff --git a/Engine/Source/Runtime/Launch/Private/LaunchEngineLoop.cpp b/Engine/Source/Runtime/Launch/Private/LaunchEngineLoop.cpp
--- a/Engine/Source/Runtime/Launch/Private/LaunchEngineLoop.cpp
+++ b/Engine/Source/Runtime/Launch/Private/LaunchEngineLoop.cpp
@@ -1436,10 +1436,6 @@ int32 FEngineLoop::PreInit( const TCHAR* CmdLine )
 
 	SlowTask.EnterProgressFrame(10);
 
-#if USE_LOCALIZED_PACKAGE_CACHE
-	FPackageLocalizationManager::Get().InitializeFromCache(MakeShareable(new FEnginePackageLocalizationCache()));
-#endif	// USE_LOCALIZED_PACKAGE_CACHE
-
 	// Initialize the RHI.
 	RHIInit(bHasEditorToken);
 
@@ -1496,6 +1492,11 @@ int32 FEngineLoop::PreInit( const TCHAR* CmdLine )
 		UMaterialInterface::AssertDefaultMaterialsPostLoaded();
 	}
 
+#if USE_LOCALIZED_PACKAGE_CACHE
+	// this loads a UObject, so we do this now after objects are brought up (otherwise, it can't load properties from config)
+	FPackageLocalizationManager::Get().InitializeFromCache(MakeShareable(new FEnginePackageLocalizationCache()));
+#endif	// USE_LOCALIZED_PACKAGE_CACHE
+	
 	// Initialize the texture streaming system (needs to happen after RHIInit and ProcessNewlyLoadedUObjects).
 	IStreamingManager::Get();
 
