Fixed VREditor regression with server not starting up (from CL 2943783)

diff --git a/Engine/Source/Runtime/Launch/Private/LaunchEngineLoop.cpp b/Engine/Source/Runtime/Launch/Private/LaunchEngineLoop.cpp
--- a/Engine/Source/Runtime/Launch/Private/LaunchEngineLoop.cpp
+++ b/Engine/Source/Runtime/Launch/Private/LaunchEngineLoop.cpp
@@ -2066,9 +2066,12 @@ bool FEngineLoop::LoadStartupCoreModules()
 
 	// Ability tasks are based on GameplayTasks, so we need to make sure that module is loaded as well
 	FModuleManager::Get().LoadModule(TEXT("GameplayTasksEditor"));
-	
-	// VREditor needs to be loaded in editor builds early, so engine content Blueprints can be loaded during DDC generation
-	FModuleManager::Get().LoadModule(TEXT("VREditor"));
+
+	if( !IsRunningDedicatedServer() )
+	{
+		// VREditor needs to be loaded in non-server editor builds early, so engine content Blueprints can be loaded during DDC generation
+		FModuleManager::Get().LoadModule(TEXT("VREditor"));
+	}
 	// -----------------------------------------------------
 
 	// HACK: load AbilitySystem editor as early as possible for statically initialized assets (non cooked BT assets needs it)
