Guard against warnings/errors changing during shutdown (still not thread-safe...)

diff --git a/Engine/Source/Runtime/Launch/Private/LaunchEngineLoop.cpp b/Engine/Source/Runtime/Launch/Private/LaunchEngineLoop.cpp
--- a/Engine/Source/Runtime/Launch/Private/LaunchEngineLoop.cpp
+++ b/Engine/Source/Runtime/Launch/Private/LaunchEngineLoop.cpp
@@ -1813,8 +1813,12 @@ int32 FEngineLoop::PreInit( const TCHAR* CmdLine )
 					
 					SET_WARN_COLOR(COLOR_RED);
 					ShownMessages.Empty(MaxMessagesToShow);
-					for (const FString& ErrorMessage : GWarn->Errors)
+					int ErrorIndex = 0;
+
+					while (ErrorIndex < GWarn->Errors.Num())
 					{
+						FString ErrorMessage = GWarn->Errors[ErrorIndex];
+
 						bool bAlreadyShown = false;
 						ShownMessages.Add(ErrorMessage, &bAlreadyShown);
 
@@ -1829,12 +1833,18 @@ int32 FEngineLoop::PreInit( const TCHAR* CmdLine )
 
 							UE_LOG(LogInit, Display, TEXT("%s"), *ErrorMessage);
 						}
+
+						ErrorIndex++;
 					}
 
 					SET_WARN_COLOR(COLOR_YELLOW);
 					ShownMessages.Empty(MaxMessagesToShow);
-					for (const FString& WarningMessage : GWarn->Warnings)
+
+					int WarnIndex = 0;
+					while (WarnIndex < GWarn->Warnings.Num())
 					{
+						FString WarningMessage = GWarn->Warnings[WarnIndex];
+
 						bool bAlreadyShown = false;
 						ShownMessages.Add(WarningMessage, &bAlreadyShown);
 
@@ -1849,6 +1859,8 @@ int32 FEngineLoop::PreInit( const TCHAR* CmdLine )
 
 							UE_LOG(LogInit, Display, TEXT("%s"), *WarningMessage);
 						}
+
+						++WarnIndex;
 					}
 				}
 
