Making FFeedbackContext's Errors and Warnings list thread safe (safe way of fixing crash on cook exit)

diff --git a/Engine/Plugins/NetcodeUnitTest/NetcodeUnitTest/Source/NetcodeUnitTest/Private/UnitTestCommandlet.cpp b/Engine/Plugins/NetcodeUnitTest/NetcodeUnitTest/Source/NetcodeUnitTest/Private/UnitTestCommandlet.cpp
--- a/Engine/Plugins/NetcodeUnitTest/NetcodeUnitTest/Source/NetcodeUnitTest/Private/UnitTestCommandlet.cpp
+++ b/Engine/Plugins/NetcodeUnitTest/NetcodeUnitTest/Source/NetcodeUnitTest/Private/UnitTestCommandlet.cpp
@@ -218,7 +218,7 @@ int32 UUnitTestCommandlet::Main(const FString& Params)
 	}
 #endif
 
-	return (GWarn->Errors.Num() == 0 ? 0 : 1);
+	return (GWarn->GetNumErrors() == 0 ? 0 : 1);
 }
 
 void UUnitTestCommandlet::CreateCustomEngine(const FString& Params)
diff --git a/Engine/Plugins/Runtime/PacketHandlers/CompressionComponents/Oodle/Source/OodleHandlerComponent/Private/OodleTrainerCommandlet.cpp b/Engine/Plugins/Runtime/PacketHandlers/CompressionComponents/Oodle/Source/OodleHandlerComponent/Private/OodleTrainerCommandlet.cpp
--- a/Engine/Plugins/Runtime/PacketHandlers/CompressionComponents/Oodle/Source/OodleHandlerComponent/Private/OodleTrainerCommandlet.cpp
+++ b/Engine/Plugins/Runtime/PacketHandlers/CompressionComponents/Oodle/Source/OodleHandlerComponent/Private/OodleTrainerCommandlet.cpp
@@ -181,7 +181,7 @@ int32 UOodleTrainerCommandlet::Main(const FString& Params)
 #endif
 
 
-	return (GWarn->Errors.Num() == 0 ? 0 : 1);
+	return (GWarn->GetNumErrors() == 0 ? 0 : 1);
 }
 
 #if HAS_OODLE_SDK
diff --git a/Engine/Source/Developer/DesktopPlatform/Private/Windows/WindowsNativeFeedbackContext.cpp b/Engine/Source/Developer/DesktopPlatform/Private/Windows/WindowsNativeFeedbackContext.cpp
--- a/Engine/Source/Developer/DesktopPlatform/Private/Windows/WindowsNativeFeedbackContext.cpp
+++ b/Engine/Source/Developer/DesktopPlatform/Private/Windows/WindowsNativeFeedbackContext.cpp
@@ -45,7 +45,7 @@ void FWindowsNativeFeedbackContext::Serialize( const TCHAR* V, ELogVerbosity::Ty
 			// Only store off the message if running a commandlet.
 			if ( IsRunningCommandlet() )
 			{
-				Errors.Add(Format);
+				AddError(Format);
 			}
 		}
 		else
@@ -53,7 +53,7 @@ void FWindowsNativeFeedbackContext::Serialize( const TCHAR* V, ELogVerbosity::Ty
 			// Only store off the message if running a commandlet.
 			if ( IsRunningCommandlet() )
 			{
-				Warnings.Add(Format);
+				AddWarning(Format);
 			}
 		}
 	}
diff --git a/Engine/Source/Developer/HotReload/Private/HotReload.cpp b/Engine/Source/Developer/HotReload/Private/HotReload.cpp
--- a/Engine/Source/Developer/HotReload/Private/HotReload.cpp
+++ b/Engine/Source/Developer/HotReload/Private/HotReload.cpp
@@ -736,9 +736,10 @@ ECompilationResult::Type FHotReloadModule::DoHotReloadInternal(const TArray<FRec
 
 	ModuleManager.ResetModulePathsCache();
 
+	
 	FFeedbackContext& ErrorsFC = UClass::GetDefaultPropertiesFeedbackContext();
-	ErrorsFC.Errors.Empty();
-	ErrorsFC.Warnings.Empty();
+	ErrorsFC.ClearWarningsAndErrors();
+
 	// Rebind the hot reload DLL 
 	TGuardValue<bool> GuardIsHotReload(GIsHotReload, true);
 	TGuardValue<bool> GuardIsInitialLoad(GIsInitialLoad, true);
@@ -794,19 +795,15 @@ ECompilationResult::Type FHotReloadModule::DoHotReloadInternal(const TArray<FRec
 		}
 	}
 
-	if (ErrorsFC.Errors.Num() || ErrorsFC.Warnings.Num())
+	if (ErrorsFC.GetNumErrors() || ErrorsFC.GetNumWarnings())
 	{
-		TArray<FString> All;
-		All = ErrorsFC.Errors;
-		All += ErrorsFC.Warnings;
-
-		ErrorsFC.Errors.Empty();
-		ErrorsFC.Warnings.Empty();
+		TArray<FString> AllErrorsAndWarnings;
+		ErrorsFC.GetErrorsAndWarningsAndEmpty(AllErrorsAndWarnings);
 
 		FString AllInOne;
-		for (int32 Index = 0; Index < All.Num(); Index++)
+		for (const FString& ErrorOrWarning : AllErrorsAndWarnings)
 		{
-			AllInOne += All[Index];
+			AllInOne += ErrorOrWarning;
 			AllInOne += TEXT("\n");
 		}
 		HotReloadAr.Logf(ELogVerbosity::Warning, TEXT("Some classes could not be reloaded:\n%s"), *AllInOne);
diff --git a/Engine/Source/Editor/UnrealEd/Private/Commandlets/ContentCommandlets.cpp b/Engine/Source/Editor/UnrealEd/Private/Commandlets/ContentCommandlets.cpp
--- a/Engine/Source/Editor/UnrealEd/Private/Commandlets/ContentCommandlets.cpp
+++ b/Engine/Source/Editor/UnrealEd/Private/Commandlets/ContentCommandlets.cpp
@@ -352,7 +352,7 @@ void UResavePackagesCommandlet::LoadAndSaveOnePackage(const FString& Filename)
 
 		static int32 LastErrorCount = 0;
 
-		int32 NumErrorsFromLoading = GWarn->Errors.Num();
+		int32 NumErrorsFromLoading = GWarn->GetNumErrors();
 		if (NumErrorsFromLoading > LastErrorCount)
 		{
 			UE_LOG(LogContentCommandlet, Warning, TEXT("%d total errors encountered during loading"), NumErrorsFromLoading);
@@ -415,9 +415,9 @@ void UResavePackagesCommandlet::LoadAndSaveOnePackage(const FString& Filename)
 			// if we do then we want to resave this package
 			if( !bSavePackage && FParse::Param(FCommandLine::Get(),TEXT("SavePackagesThatHaveFailedLoads")) == true )
 			{
-				//UE_LOG(LogContentCommandlet, Warning, TEXT( "NumErrorsFromLoading: %d GWarn->Errors num: %d" ), NumErrorsFromLoading, GWarn->Errors.Num() );
+				//UE_LOG(LogContentCommandlet, Warning, TEXT( "NumErrorsFromLoading: %d GWarn->Errors num: %d" ), NumErrorsFromLoading, GWarn->GetNumErrors() );
 
-				if( NumErrorsFromLoading != GWarn->Errors.Num() )
+				if( NumErrorsFromLoading != GWarn->GetNumErrors() )
 				{
 					bSavePackage = true;
 				}
diff --git a/Engine/Source/Runtime/Core/Private/Linux/LinuxPlatformFeedbackContextPrivate.h b/Engine/Source/Runtime/Core/Private/Linux/LinuxPlatformFeedbackContextPrivate.h
--- a/Engine/Source/Runtime/Core/Private/Linux/LinuxPlatformFeedbackContextPrivate.h
+++ b/Engine/Source/Runtime/Core/Private/Linux/LinuxPlatformFeedbackContextPrivate.h
@@ -47,7 +47,7 @@ public:
 				// Only store off the message if running a commandlet.
 				if (IsRunningCommandlet())
 				{
-					Errors.Add(Format);
+					AddError(Format);
 				}
 
 				// send errors (warnings are too spammy) to syslog too (for zabbix etc)
@@ -58,7 +58,7 @@ public:
 				// Only store off the message if running a commandlet.
 				if ( IsRunningCommandlet() )
 				{
-					Warnings.Add(Format);
+					AddWarning(Format);
 				}
 			}
 		}
diff --git a/Engine/Source/Runtime/Core/Private/Mac/MacPlatformFeedbackContextPrivate.h b/Engine/Source/Runtime/Core/Private/Mac/MacPlatformFeedbackContextPrivate.h
--- a/Engine/Source/Runtime/Core/Private/Mac/MacPlatformFeedbackContextPrivate.h
+++ b/Engine/Source/Runtime/Core/Private/Mac/MacPlatformFeedbackContextPrivate.h
@@ -46,7 +46,7 @@ public:
 				// Only store off the message if running a commandlet.
 				if ( IsRunningCommandlet() )
 				{
-					Errors.Add(Format);
+					AddError(Format);
 				}
 			}
 			else
@@ -55,7 +55,7 @@ public:
 				// Only store off the message if running a commandlet.
 				if ( IsRunningCommandlet() )
 				{
-					Warnings.Add(Format);
+					AddWarning(Format);
 				}
 			}
 		}
diff --git a/Engine/Source/Runtime/Core/Private/WinRT/WinRTFeedbackContextPrivate.h b/Engine/Source/Runtime/Core/Private/WinRT/WinRTFeedbackContextPrivate.h
--- a/Engine/Source/Runtime/Core/Private/WinRT/WinRTFeedbackContextPrivate.h
+++ b/Engine/Source/Runtime/Core/Private/WinRT/WinRTFeedbackContextPrivate.h
@@ -48,7 +48,7 @@ public:
 				// Only store off the message if running a commandlet.
 				if ( IsRunningCommandlet() )
 				{
-					Errors.Add(Format);
+					AddError(Format);
 				}
 			}
 			else
@@ -57,7 +57,7 @@ public:
 				// Only store off the message if running a commandlet.
 				if ( IsRunningCommandlet() )
 				{
-					Warnings.Add(Format);
+					AddWarning(Format);
 				}
 			}
 		}
diff --git a/Engine/Source/Runtime/Core/Private/Windows/WindowsPlatformFeedbackContextPrivate.h b/Engine/Source/Runtime/Core/Private/Windows/WindowsPlatformFeedbackContextPrivate.h
--- a/Engine/Source/Runtime/Core/Private/Windows/WindowsPlatformFeedbackContextPrivate.h
+++ b/Engine/Source/Runtime/Core/Private/Windows/WindowsPlatformFeedbackContextPrivate.h
@@ -44,7 +44,7 @@ public:
 				// Only store off the message if running a commandlet.
 				if ( IsRunningCommandlet() )
 				{
-					Errors.Add(Format);
+					AddError(Format);
 				}
 			}
 			else
@@ -52,7 +52,7 @@ public:
 				// Only store off the message if running a commandlet.
 				if ( IsRunningCommandlet() )
 				{
-					Warnings.Add(Format);
+					AddWarning(Format);
 				}
 			}
 		}
diff --git a/Engine/Source/Runtime/Core/Public/HAL/FeedbackContextAnsi.h b/Engine/Source/Runtime/Core/Public/HAL/FeedbackContextAnsi.h
--- a/Engine/Source/Runtime/Core/Public/HAL/FeedbackContextAnsi.h
+++ b/Engine/Source/Runtime/Core/Public/HAL/FeedbackContextAnsi.h
@@ -57,11 +57,11 @@ public:
 			{
 				if(Verbosity == ELogVerbosity::Error)
 				{
-					Errors.Add(Format);
+					AddError(Format);
 				}
 				else
 				{
-					Warnings.Add(Format);
+					AddWarning(Format);
 				}
 			}
 			LocalPrint(*Format);
diff --git a/Engine/Source/Runtime/Core/Public/Misc/FeedbackContext.h b/Engine/Source/Runtime/Core/Public/Misc/FeedbackContext.h
--- a/Engine/Source/Runtime/Core/Public/Misc/FeedbackContext.h
+++ b/Engine/Source/Runtime/Core/Public/Misc/FeedbackContext.h
@@ -78,9 +78,6 @@ public:
 	virtual TWeakPtr<class SBuildProgressWidget> ShowBuildProgressWindow() {return TWeakPtr<class SBuildProgressWidget>();}
 	virtual void CloseBuildProgressWindow() {}
 
-	TArray<FString> Warnings;
-	TArray<FString> Errors;
-
 	bool	TreatWarningsAsErrors;
 
 	FFeedbackContext()
@@ -90,10 +87,54 @@ public:
 
 	virtual ~FFeedbackContext();
 
+	/** Gets warnings history */
+	void GetWarnings(TArray<FString>& OutWarnings) const
+	{
+		FScopeLock WarningsAndErrorsLock(&WarningsAndErrorsCritical);
+		OutWarnings = Warnings;
+	}
+	int32 GetNumWarnings() const
+	{
+		return Warnings.Num();
+	}
+
+	/** Gets errors history */
+	void GetErrors(TArray<FString>& OutErrors) const
+	{
+		FScopeLock WarningsAndErrorsLock(&WarningsAndErrorsCritical);
+		OutErrors = Errors;
+	}
+	int32 GetNumErrors() const
+	{
+		return Errors.Num();
+	}
+
+	/** Gets all errors and warnings and clears the history */
+	void GetErrorsAndWarningsAndEmpty(TArray<FString>& OutWarningsAndErrors)
+	{
+		FScopeLock WarningsAndErrorsLock(&WarningsAndErrorsCritical);
+		OutWarningsAndErrors = MoveTemp(Errors);
+		OutWarningsAndErrors += MoveTemp(Warnings);
+	}
+	/** Clears all history */
+	void ClearWarningsAndErrors()
+	{
+		FScopeLock WarningsAndErrorsLock(&WarningsAndErrorsCritical);
+		Errors.Empty();
+		Warnings.Empty();
+	}
+
 private:
 	FFeedbackContext(const FFeedbackContext&);
 	FFeedbackContext& operator=(const FFeedbackContext&);
 
+	/** Warnings history */
+	TArray<FString> Warnings;
+	/** Errors history */
+	TArray<FString> Errors;
+	/** Guard for the errors and warnings history */
+	mutable FCriticalSection WarningsAndErrorsCritical;
+
 protected:
 	
 	friend FSlowTask;
@@ -108,6 +149,25 @@ protected:
 	/** Update the UI as a result of the scope stack changing */
 	void UpdateUI();
 
+	/**
+	 * Adds a new warning message to warnings history.
+	 * @param InWarning Warning message
+	 */
+	void AddWarning(const FString& InWarning)
+	{
+		FScopeLock WarningsAndErrorsLock(&WarningsAndErrorsCritical);
+		Warnings.Add(InWarning);
+	}
+
+	/**
+	* Adds a new error message to errors history.
+	* @param InWarning Error message
+	*/
+	void AddError(const FString& InError)
+	{
+		FScopeLock WarningsAndErrorsLock(&WarningsAndErrorsCritical);
+		Errors.Add(InError);
+	}
 };
 
 /** Enum to specify a particular slow task section should be shown */
diff --git a/Engine/Source/Runtime/CoreUObject/Private/UObject/Class.cpp b/Engine/Source/Runtime/CoreUObject/Private/UObject/Class.cpp
--- a/Engine/Source/Runtime/CoreUObject/Private/UObject/Class.cpp
+++ b/Engine/Source/Runtime/CoreUObject/Private/UObject/Class.cpp
@@ -2881,18 +2881,22 @@ public:
 
 			if(Verbosity == ELogVerbosity::Error)
 			{
-				Errors.Add(Format);
+				AddError(Format);
 			}
 			else
 			{
-				Warnings.Add(Format);
+				AddWarning(Format);
 			}
 		}
 
-		if( GLogConsole )
-			GLogConsole->Serialize( V, Verbosity, Category );
-		if( !GLog->IsRedirectingTo( this ) )
-			GLog->Serialize( V, Verbosity, Category );
+		if (GLogConsole)
+		{
+			GLogConsole->Serialize(V, Verbosity, Category);
+		}
+		if (!GLog->IsRedirectingTo(this))
+		{
+			GLog->Serialize(V, Verbosity, Category);
+		}
 	}
 
 	FContextSupplier* GetContext() const
diff --git a/Engine/Source/Runtime/CoreUObject/Private/UObject/UObjectBase.cpp b/Engine/Source/Runtime/CoreUObject/Private/UObject/UObjectBase.cpp
--- a/Engine/Source/Runtime/CoreUObject/Private/UObject/UObjectBase.cpp
+++ b/Engine/Source/Runtime/CoreUObject/Private/UObject/UObjectBase.cpp
@@ -727,21 +727,17 @@ static void UObjectLoadAllCompiledInDefaultProperties()
 			Class->GetDefaultObject();
 		}
 		FFeedbackContext& ErrorsFC = UClass::GetDefaultPropertiesFeedbackContext();
-		if (ErrorsFC.Errors.Num() || ErrorsFC.Warnings.Num())
+		if (ErrorsFC.GetNumErrors() || ErrorsFC.GetNumWarnings())
 		{
-			TArray<FString> All;
-			All = ErrorsFC.Errors;
-			All += ErrorsFC.Warnings;
-
-			ErrorsFC.Errors.Empty();
-			ErrorsFC.Warnings.Empty();
+			TArray<FString> AllErrorsAndWarnings;
+			ErrorsFC.GetErrorsAndWarningsAndEmpty(AllErrorsAndWarnings);
 
 			FString AllInOne;
 			UE_LOG(LogUObjectBase, Warning, TEXT("-------------- Default Property warnings and errors:"));
-			for (int32 Index = 0; Index < All.Num(); Index++)
+			for (const FString& ErrorOrWarning : AllErrorsAndWarnings)
 			{
-				UE_LOG(LogUObjectBase, Warning, TEXT("%s"), *All[Index]);
-				AllInOne += All[Index];
+				UE_LOG(LogUObjectBase, Warning, TEXT("%s"), *ErrorOrWarning);
+				AllInOne += ErrorOrWarning;
 				AllInOne += TEXT("\n");
 			}
 			FMessageDialog::Open(EAppMsgType::Ok, FText::Format( NSLOCTEXT("Core", "DefaultPropertyWarningAndErrors", "Default Property warnings and errors:\n{0}"), FText::FromString( AllInOne ) ) );
diff --git a/Engine/Source/Runtime/Launch/Private/LaunchEngineLoop.cpp b/Engine/Source/Runtime/Launch/Private/LaunchEngineLoop.cpp
--- a/Engine/Source/Runtime/Launch/Private/LaunchEngineLoop.cpp
+++ b/Engine/Source/Runtime/Launch/Private/LaunchEngineLoop.cpp
@@ -1799,7 +1799,12 @@ int32 FEngineLoop::PreInit( const TCHAR* CmdLine )
 			// Log warning/ error summary.
 			if( Commandlet->ShowErrorCount )
 			{
-				if( GWarn->Errors.Num() || GWarn->Warnings.Num() )
+				TArray<FString> AllErrors;
+				TArray<FString> AllWarnings;
+				GWarn->GetErrors(AllErrors);
+				GWarn->GetWarnings(AllWarnings);
+
+				if (AllErrors.Num() || AllWarnings.Num())
 				{
 					SET_WARN_COLOR(COLOR_WHITE);
 					UE_LOG(LogInit, Display, TEXT(""));
@@ -1807,18 +1812,15 @@ int32 FEngineLoop::PreInit( const TCHAR* CmdLine )
 					UE_LOG(LogInit, Display, TEXT("-----------------------------------"));
 
 					const int32 MaxMessagesToShow = (GIsBuildMachine || FParse::Param(FCommandLine::Get(), TEXT("DUMPALLWARNINGS"))) ? 
-						FMath::Max(GWarn->Errors.Num(), GWarn->Warnings.Num()) : 50;
-
-					TSet<FString> ShownMessages;
+						FMath::Max(AllErrors.Num(), AllWarnings.Num()) : 50;
 					
-					SET_WARN_COLOR(COLOR_RED);
+					TSet<FString> ShownMessages;					
 					ShownMessages.Empty(MaxMessagesToShow);
-					int ErrorIndex = 0;
 
-					while (ErrorIndex < GWarn->Errors.Num())
-					{
-						FString ErrorMessage = GWarn->Errors[ErrorIndex];
+					SET_WARN_COLOR(COLOR_RED);
 
+					for (const FString& ErrorMessage : AllErrors)
+					{
 						bool bAlreadyShown = false;
 						ShownMessages.Add(ErrorMessage, &bAlreadyShown);
 
@@ -1827,24 +1829,19 @@ int32 FEngineLoop::PreInit( const TCHAR* CmdLine )
 							if (ShownMessages.Num() > MaxMessagesToShow)
 							{
 								SET_WARN_COLOR(COLOR_WHITE);
-								UE_CLOG(MaxMessagesToShow < GWarn->Errors.Num(), LogInit, Display, TEXT("NOTE: Only first %d errors displayed."), MaxMessagesToShow);
+								UE_CLOG(MaxMessagesToShow < AllErrors.Num(), LogInit, Display, TEXT("NOTE: Only first %d errors displayed."), MaxMessagesToShow);
 								break;
 							}
 
 							UE_LOG(LogInit, Display, TEXT("%s"), *ErrorMessage);
 						}
-
-						ErrorIndex++;
 					}
 
 					SET_WARN_COLOR(COLOR_YELLOW);
 					ShownMessages.Empty(MaxMessagesToShow);
 
-					int WarnIndex = 0;
-					while (WarnIndex < GWarn->Warnings.Num())
+					for (const FString& WarningMessage : AllWarnings)
 					{
-						FString WarningMessage = GWarn->Warnings[WarnIndex];
-
 						bool bAlreadyShown = false;
 						ShownMessages.Add(WarningMessage, &bAlreadyShown);
 
@@ -1853,14 +1850,12 @@ int32 FEngineLoop::PreInit( const TCHAR* CmdLine )
 							if (ShownMessages.Num() > MaxMessagesToShow)
 							{
 								SET_WARN_COLOR(COLOR_WHITE);
-								UE_CLOG(MaxMessagesToShow < GWarn->Warnings.Num(), LogInit, Display, TEXT("NOTE: Only first %d warnings displayed."), MaxMessagesToShow);
+								UE_CLOG(MaxMessagesToShow < AllWarnings.Num(), LogInit, Display, TEXT("NOTE: Only first %d warnings displayed."), MaxMessagesToShow);
 								break;
 							}
 
 							UE_LOG(LogInit, Display, TEXT("%s"), *WarningMessage);
 						}
-
-						++WarnIndex;
 					}
 				}
 
@@ -1870,17 +1865,17 @@ int32 FEngineLoop::PreInit( const TCHAR* CmdLine )
 				{
 					SET_WARN_COLOR(COLOR_RED);
 					UE_LOG(LogInit, Display, TEXT("Commandlet->Main return this error code: %d"), ErrorLevel );
-					UE_LOG(LogInit, Display, TEXT("With %d error(s), %d warning(s)"), GWarn->Errors.Num(), GWarn->Warnings.Num() );
+					UE_LOG(LogInit, Display, TEXT("With %d error(s), %d warning(s)"), AllErrors.Num(), AllWarnings.Num() );
 				}
-				else if( ( GWarn->Errors.Num() == 0 ) )
+				else if( ( AllErrors.Num() == 0 ) )
 				{
-					SET_WARN_COLOR(GWarn->Warnings.Num() ? COLOR_YELLOW : COLOR_GREEN);
-					UE_LOG(LogInit, Display, TEXT("Success - %d error(s), %d warning(s)"), GWarn->Errors.Num(), GWarn->Warnings.Num() );
+					SET_WARN_COLOR(AllWarnings.Num() ? COLOR_YELLOW : COLOR_GREEN);
+					UE_LOG(LogInit, Display, TEXT("Success - %d error(s), %d warning(s)"), AllErrors.Num(), AllWarnings.Num() );
 				}
 				else
 				{
 					SET_WARN_COLOR(COLOR_RED);
-					UE_LOG(LogInit, Display, TEXT("Failure - %d error(s), %d warning(s)"), GWarn->Errors.Num(), GWarn->Warnings.Num() );
+					UE_LOG(LogInit, Display, TEXT("Failure - %d error(s), %d warning(s)"), AllErrors.Num(), AllWarnings.Num() );
 					ErrorLevel = 1;
 				}
 				CLEAR_WARN_COLOR();
