Update version number in //UE4/Main to 4.14.

diff --git a/Engine/Build/Build.version b/Engine/Build/Build.version
--- a/Engine/Build/Build.version
+++ b/Engine/Build/Build.version
@@ -1,6 +1,6 @@
 {
 	"MajorVersion": 4,
-	"MinorVersion": 13,
+	"MinorVersion": 14,
 	"PatchVersion": 0,
 	"Changelist": 0,
 	"IsLicenseeVersion": 0,
diff --git a/Engine/Source/Runtime/Launch/Resources/Version.h b/Engine/Source/Runtime/Launch/Resources/Version.h
--- a/Engine/Source/Runtime/Launch/Resources/Version.h
+++ b/Engine/Source/Runtime/Launch/Resources/Version.h
@@ -12,7 +12,7 @@
 #if !LOCKED_BUILD_VERSION
 
 #define ENGINE_MAJOR_VERSION	4
-#define ENGINE_MINOR_VERSION	13
+#define ENGINE_MINOR_VERSION	14
 #define ENGINE_PATCH_VERSION	0
 
 #define ENGINE_IS_LICENSEE_VERSION 0
