Copying //UE4/Dev-Build to //UE4/Dev-Main (Source: //UE4/Dev-Build @ 3075869)

==========================
MAJOR FEATURES + CHANGES
==========================

Change 3059740 on 2016/07/21 by [REDACTED_PERSON]

	Format license violations on separate lines rather than as a single multi-line string, so they can be parsed out correctly by the EC log parser.

Change 3063115 on 2016/07/25 by [REDACTED_PERSON]

	Add missing projects to Mono AutomationTool solution.

Change 3067125 on 2016/07/27 by [REDACTED_PERSON]

	Clean up and document the engine versioning system.

	* Version.h has a summary of all the different versioning concepts in UE4, as well as an explanation of what all the macros do.
	* VersionLocked.h has been removed. ENGINE_NET_VERSION and ENGINE_REPLAY_VERSION can be overriden directly in Version.h if necessary.
	* Added explicit macros for ENGINE_CURRENT_CL_VERSION And ENGINE_COMPATIBLE_CL_VERSION. The MODULE_API_VERSION macro is now defined to be the same as ENGINE_COMPATIBLE_CL_VERSION.
	* Build/Build.version now includes the compatible CL version, so it can be determined without having to parse header files inside MacToolChain.cs.
	* UpdateLocalVersion UAT command now exposes a -CompatibleCL=... argument to update ENGINE_COMPATIBLE_CL_VERSION and the CompatibleChange property in Build.version.
	* SetVersion BuildGraph task now has a CompatibleChange="" attribute to update ENGINE_COMPATIBLE_CL_VERSION and the CompatibleChange property in Build.version.
	* VersionFileUpdater now makes an effort to parse tokens matching the macros it's going to overwrite, rather than requiring exact matches in whitespace in Version.h.
	* Running UpdateLocalVersion now syncs files to revision 0 before writing if P4 is enabled, similarly to UGS, or fails if they are not writable.

Change 3067242 on 2016/07/27 by [REDACTED_PERSON]

	Fix UGS showing a build failure/warning dialog immediately if you submit while the build is yellow/red. Now ignores builds before your last submitted CL.

Change 3068116 on 2016/07/28 by [REDACTED_PERSON]

	Add a small Windows utility to terminate all running processes under a given root directory, and run it before and after every sync. Mean to fix instances where zombie editor processes do not terminate correctly and cause subsequent syncs to fail.

Change 3068246 on 2016/07/28 by [REDACTED_PERSON]

	UGS: Improvements to 'Clean Workspace' dialog - empty folders are now set to be deleted by default, and added a context menu to select all/none/empty/default for a subtree, as well as to open Windows explorer at that location.

Change 3068573 on 2016/07/28 by [REDACTED_PERSON]

	Attempt to fix AppDomain unload errors; suspect caused by delay to terminate child processes when managed exception is thrown on child threads. Now waits for the currently running threads to finish executing before quitting.

Change 3068575 on 2016/07/28 by [REDACTED_PERSON]

	Respect the -noxge command-line argument when compiling UHT.

Change 3068665 on 2016/07/28 by [REDACTED_PERSON]

	Delete any DLLs in output folders which share names with build products. The Windows loader reads DLLs from the first location it finds a file with a matching name, so we need to ensure that it doesn't load stale DLLs when output directories are changed (moving a module into a plugin, for example).

Change 3073316 on 2016/08/02 by [REDACTED_PERSON]

	CoreUObject: Replace header guards with #pragma once directives.

Change 3073325 on 2016/08/02 by [REDACTED_PERSON]

	CoreUObject: Manually untangle a few circular includes around ObjectBase.h and Class.h to allow making them into IWYU-style headers.

Change 3074391 on 2016/08/02 by [REDACTED_PERSON]

	GitHub #2646: Always allow programs to build as part of non-game solution configurations.

Change 3075123 on 2016/08/03 by [REDACTED_PERSON]

	Ensure that zip files created outside of Ionic Zip can be unzipped with CommandUtils.UnzipFiles

	Zip files created with 7-zip, WinRAR, and Windows Explorer all contain entities for directories, not just the files in them.  Code in CommandUtils.UnzipFiles made the assumption that all entities in the zip file were files, and an exception was being thrown on encountering a directory.

Change 3075848 on 2016/08/03 by [REDACTED_PERSON]

	UBT: Fix a few cases where it's possible to construct a non-canonical DirectoryReference - root directories on Windows must always have a terminating slash, but any other directories must note. DirectoryInfo and Path.GetFullPath() do not enforce this requirement.

Change 3075850 on 2016/08/03 by [REDACTED_PERSON]

	UBT: Fix UHT failures when running with -nopch.

Change 3077058 by [REDACTED_PERSON]

	UHT: Remove the auto-generated ObjectBase.h include from generated header files. Was originally added to allow doxygen to parse headers more easily, but was switched to parsing the file only as part of parent files instead.

diff --git a/Engine/Build/Build.version b/Engine/Build/Build.version
--- a/Engine/Build/Build.version
+++ b/Engine/Build/Build.version
@@ -3,6 +3,7 @@
 	"MinorVersion": 14,
 	"PatchVersion": 0,
 	"Changelist": 0,
+	"CompatibleChangelist": 0,
 	"IsLicenseeVersion": 0,
 	"BranchName": ""
 }
\ No newline at end of file
diff --git a/Engine/Source/Programs/AutomationTool/AutomationUtils/CommandUtils.cs b/Engine/Source/Programs/AutomationTool/AutomationUtils/CommandUtils.cs
--- a/Engine/Source/Programs/AutomationTool/AutomationUtils/CommandUtils.cs
+++ b/Engine/Source/Programs/AutomationTool/AutomationUtils/CommandUtils.cs
@@ -946,6 +946,23 @@ namespace AutomationTool
 			File.SetLastWriteTimeUtc(Filename, Date);
 		}
 
+		/// <summary>
+		/// Determines whether the given file is read-only
+		/// </summary>
+		/// <param name="Filename">Filename</param>
+		/// <returns>True if the file is read-only</returns>
+		public static bool IsReadOnly(string Filename)
+		{
+			Filename = ConvertSeparators(PathSeparator.Default, Filename);
+			if (!File.Exists(Filename))
+			{
+				throw new AutomationException(new FileNotFoundException("File not found.", Filename), "Unable to set attributes for a non-existing file.");
+			}
+
+			FileAttributes Attributes = File.GetAttributes(Filename);
+			return (Attributes & FileAttributes.ReadOnly) != 0;
+		}
+
 		/// <summary>
 		/// Sets file attributes. Will not change attributes that have not been specified.
 		/// </summary>
@@ -1836,6 +1853,25 @@ namespace AutomationTool
 			return int.Parse(num);
 		}
 
+		/// <summary>
+		/// Parses the command's Params list for a parameter and reads its value. 
+		/// Ex. ParseParamValue(Args, "map=")
+		/// </summary>
+		/// <param name="Param">Param to read its value.</param>
+		/// <returns>Returns the value or Default if the parameter was not found.</returns>
+		public int? ParseParamNullableInt(string Param)
+		{
+			string Value = ParseParamValue(Params, Param, null);
+			if(Value == null)
+			{
+				return null;
+			}
+			else
+			{
+				return int.Parse(Value);
+			}
+		}
+
 		/// <summary>
 		/// Makes sure path can be used as a command line param (adds quotes if it contains spaces)
 		/// </summary>
@@ -2120,12 +2156,12 @@ namespace AutomationTool
 		/// <returns>List of files written</returns>
 		public static IEnumerable<string> UnzipFiles(string ZipFileName, string BaseDirectory)
 		{
-            // manually extract the files. There was a problem with the Ionic.Zip library that required this on non-PC at one point,
-            // but that problem is now fixed. Leaving this code as is as we need to return the list of created files and fix up their permissions anyway.
-            using (Ionic.Zip.ZipFile Zip = new Ionic.Zip.ZipFile(ZipFileName))
+			// manually extract the files. There was a problem with the Ionic.Zip library that required this on non-PC at one point,
+			// but that problem is now fixed. Leaving this code as is as we need to return the list of created files and fix up their permissions anyway.
+			using (Ionic.Zip.ZipFile Zip = new Ionic.Zip.ZipFile(ZipFileName))
 			{
 				List<string> OutputFileNames = new List<string>();
-				foreach(Ionic.Zip.ZipEntry Entry in Zip.Entries)
+				foreach(Ionic.Zip.ZipEntry Entry in Zip.Entries.Where(x => !x.IsDirectory))
 				{
 					string OutputFileName = Path.Combine(BaseDirectory, Entry.FileName);
 					Directory.CreateDirectory(Path.GetDirectoryName(OutputFileName));
diff --git a/Engine/Source/Programs/AutomationTool/AutomationUtils/ParallelExecutor.cs b/Engine/Source/Programs/AutomationTool/AutomationUtils/ParallelExecutor.cs
--- a/Engine/Source/Programs/AutomationTool/AutomationUtils/ParallelExecutor.cs
+++ b/Engine/Source/Programs/AutomationTool/AutomationUtils/ParallelExecutor.cs
@@ -148,31 +148,34 @@ namespace AutomationTool
 							}
 
 							// Check the exit code
-							if(CompletedAction.ExitCode != 0)
+							if(CompletedAction.ExitCode == 0)
 							{
-								ExitCode = CompletedAction.ExitCode;
-								if(!CompletedAction.Action.bSkipIfProjectFailed)
+								// Mark all the dependents as done
+								foreach(BuildAction DependantAction in CompletedAction.Action.Dependants)
 								{
-									foreach(KeyValuePair<BuildActionExecutor, Thread> Pair in ExecutingActions)
+									if(--DependantAction.MissingDependencyCount == 0)
 									{
-										Pair.Value.Abort();
-										Pair.Value.Join();
+										QueuedActions.Add(DependantAction);
 									}
-									return ExitCode;
 								}
 							}
-
-							// Mark all the dependents as done
-							foreach(BuildAction DependantAction in CompletedAction.Action.Dependants)
+							else
 							{
-								if(--DependantAction.MissingDependencyCount == 0)
+								// Update the exit code if it's not already set
+								if(ExitCode == 0)
 								{
-									QueuedActions.Add(DependantAction);
+									ExitCode = CompletedAction.ExitCode;
 								}
 							}
 						}
 						CompletedActions.Clear();
 					}
+
+					// If we've already got a non-zero exit code, clear out the list of queued actions so nothing else will run
+					if(ExitCode != 0)
+					{
+						QueuedActions.Clear();
+					}
 				}
 			}
 			return ExitCode;
diff --git a/Engine/Source/Programs/AutomationTool/AutomationUtils/UE4Build.cs b/Engine/Source/Programs/AutomationTool/AutomationUtils/UE4Build.cs
--- a/Engine/Source/Programs/AutomationTool/AutomationUtils/UE4Build.cs
+++ b/Engine/Source/Programs/AutomationTool/AutomationUtils/UE4Build.cs
@@ -459,7 +459,7 @@ namespace AutomationTool
 		/// <summary>
 		/// Updates the engine version files
 		/// </summary>
-		public List<string> UpdateVersionFiles(bool ActuallyUpdateVersionFiles = true, int? ChangelistNumberOverride = null, string Build = null)
+		public List<string> UpdateVersionFiles(bool ActuallyUpdateVersionFiles = true, int? ChangelistNumberOverride = null, int? CompatibleChangelistNumberOverride = null, string Build = null)
 		{
 			bool bIsLicenseeVersion = ParseParam("Licensee");
 			bool bDoUpdateVersionFiles = CommandUtils.P4Enabled && ActuallyUpdateVersionFiles;		
@@ -468,36 +468,45 @@ namespace AutomationTool
 			{
 				ChangelistNumber = ChangelistNumberOverride.HasValue? ChangelistNumberOverride.Value : P4Env.Changelist;
 			}
+			int CompatibleChangelistNumber = ChangelistNumber;
+			if(bDoUpdateVersionFiles && CompatibleChangelistNumberOverride.HasValue)
+			{
+				CompatibleChangelistNumber = CompatibleChangelistNumberOverride.Value;
+			}
 
 			string Branch = P4Enabled ? P4Env.BuildRootEscaped : "";
-			return StaticUpdateVersionFiles(ChangelistNumber, Branch, Build, bIsLicenseeVersion, bDoUpdateVersionFiles);
+			return StaticUpdateVersionFiles(ChangelistNumber, CompatibleChangelistNumber, Branch, Build, bIsLicenseeVersion, bDoUpdateVersionFiles);
 		}
 
-		public static List<string> StaticUpdateVersionFiles(int ChangelistNumber, string Branch, string Build, bool bIsLicenseeVersion, bool bDoUpdateVersionFiles)
+		public static List<string> StaticUpdateVersionFiles(int ChangelistNumber, int CompatibleChangelistNumber, string Branch, string Build, bool bIsLicenseeVersion, bool bDoUpdateVersionFiles)
 		{
 			string ChangelistString = (ChangelistNumber != 0 && bDoUpdateVersionFiles)? ChangelistNumber.ToString() : String.Empty;
 
 			var Result = new List<String>();
 			{
-				string VerFile = CombinePaths(CmdEnv.LocalRoot, "Engine", "Build", "Build.version");
+				string VerFile = BuildVersion.GetDefaultFileName();
 				if (bDoUpdateVersionFiles)
 				{
 					LogLog("Updating {0} with:", VerFile);
 					LogLog("  Changelist={0}", ChangelistNumber);
+					LogLog("  CompatibleChangelist={0}", CompatibleChangelistNumber);
 					LogLog("  IsLicenseeVersion={0}", bIsLicenseeVersion? 1 : 0);
 					LogLog("  BranchName={0}", Branch);
 
-					string VerText = CommandUtils.ReadAllText(VerFile);
+					BuildVersion Version;
+					if(!BuildVersion.TryRead(VerFile, out Version))
+					{
+						Version = new BuildVersion();
+					}
+
+					Version.Changelist = ChangelistNumber;
+					Version.CompatibleChangelist = CompatibleChangelistNumber;
+					Version.IsLicenseeVersion = bIsLicenseeVersion? 1 : 0;
+					Version.BranchName = Branch;
 
-					JavaScriptSerializer Serializer = new JavaScriptSerializer();
-					Dictionary<string, object> Pairs = Serializer.Deserialize<Dictionary<string, object>>(VerText);
-					Pairs["Changelist"] = ChangelistNumber;
-					Pairs["IsLicenseeVersion"] = bIsLicenseeVersion? 1 : 0;
-					Pairs["BranchName"] = Branch;
-					VerText = Serializer.Serialize(Pairs).Replace("{\"", "{\n\t\"").Replace(",\"", ",\n\t\"").Replace("\":", "\": ").Replace("\"}", "\"\n}");
+					VersionFileUpdater.MakeFileWriteable(VerFile);
 
-					CommandUtils.SetFileAttributes(VerFile, ReadOnly: false);
-					CommandUtils.WriteAllText(VerFile, VerText);
+					Version.Write(VerFile);
 				}
 				else
 				{
@@ -514,6 +523,10 @@ namespace AutomationTool
 					LogLog("Updating {0} with:", VerFile);
 					LogLog(" #define	BRANCH_NAME  {0}", Branch);
 					LogLog(" #define	BUILT_FROM_CHANGELIST  {0}", ChangelistString);
+					if(CompatibleChangelistNumber > 0)
+					{
+						LogLog(" #define	ENGINE_COMPATIBLE_CL_VERSION  {0}", CompatibleChangelistNumber);
+					}
 					if (Build != null)
 					{
 						LogLog(" #define	BUILD_VERSION  {0}", Build);
@@ -527,7 +540,11 @@ namespace AutomationTool
 					{
 						VersionH.ReplaceLine("#define BUILD_VERSION ", "L\"" + Build + "\"");
 					}
-					VersionH.ReplaceOrAddLine("#define ENGINE_IS_LICENSEE_VERSION ", bIsLicenseeVersion ? "1" : "0");
+					if(CompatibleChangelistNumber > 0)
+					{
+						VersionH.ReplaceLine("#define ENGINE_COMPATIBLE_CL_VERSION", CompatibleChangelistNumber.ToString(), bIsLicenseeVersion? 0 : 1);
+					}
+					VersionH.ReplaceLine("#define ENGINE_IS_LICENSEE_VERSION ", bIsLicenseeVersion ? "1" : "0");
 
                     VersionH.Commit();
                 }
@@ -1387,7 +1404,7 @@ namespace AutomationTool
 						// When building a target for Mac or iOS, use UBT's -flushmac option to clean up the remote builder
 						bool bForceFlushMac = DeleteBuildProducts && (Target.Platform == UnrealBuildTool.UnrealTargetPlatform.Mac || Target.Platform == UnrealBuildTool.UnrealTargetPlatform.IOS);
 						LogSetProgress(InShowProgress, "Building header tool...");
-						BuildManifest Manifest = BuildWithUBT(Target.TargetName, Target.Platform, Target.Config.ToString(), Target.UprojectPath, bForceMonolithic, bForceNonUnity, bForceDebugInfo, bForceFlushMac, !CanUseXGE(Target.Platform), Target.UBTArgs, bForceUnity);
+						BuildManifest Manifest = BuildWithUBT(Target.TargetName, Target.Platform, Target.Config.ToString(), Target.UprojectPath, bForceMonolithic, bForceNonUnity, bForceDebugInfo, bForceFlushMac, !bCanUseXGE || !CanUseXGE(Target.Platform), Target.UBTArgs, bForceUnity);
 						if(InTargetToManifest != null)
 						{
 							InTargetToManifest[Target] = Manifest;
diff --git a/Engine/Source/Programs/AutomationTool/AutomationUtils/Utils.cs b/Engine/Source/Programs/AutomationTool/AutomationUtils/Utils.cs
--- a/Engine/Source/Programs/AutomationTool/AutomationUtils/Utils.cs
+++ b/Engine/Source/Programs/AutomationTool/AutomationUtils/Utils.cs
@@ -1047,44 +1047,95 @@ namespace AutomationTool
 		}
 
 		/// <summary>
-		/// Doc
+		/// Finds the line containing the same tokens as the given prefix, and change the suffix to be the given string
 		/// </summary>
-		public void ReplaceLine(string StartOfLine, string ReplacementRHS)
+		/// <param name="Prefix">Prefix to search for</param>
+		/// <param name="Suffix">Suffix to replace the rest of the line with</param>
+		/// <param name="OccurenceIdx">Index of the matching occurrence to replace</param>
+		public void ReplaceLine(string Prefix, string Suffix, int OccurenceIdx = 0)
 		{
-			for (int Index = 0; Index < Lines.Count; ++Index)
+			int MatchIdx = 0;
+			for(int LineIdx = 0; LineIdx < Lines.Count; LineIdx++)
 			{
-				if (Lines[Index].StartsWith(StartOfLine))
+				int LineOffset = 0;
+				if(MatchCppTokens(Lines[LineIdx], ref LineOffset, Prefix))
 				{
-					Lines[Index] = StartOfLine + ReplacementRHS;
-					return;
+					if(OccurenceIdx == MatchIdx)
+					{
+						while(LineOffset < Lines[LineIdx].Length && Char.IsWhiteSpace(Lines[LineIdx][LineOffset]))
+						{
+							LineOffset++;
+						}
+						Lines[LineIdx] = Lines[LineIdx].Substring(0, LineOffset) + Suffix;
+						return;
+					}
+					MatchIdx++;
 				}
 			}
-            throw new AutomationException("Unable to find line {0} in {1}", StartOfLine, MyFile.FullName);
+            throw new AutomationException("Unable to find line {0} in {1}", Prefix, MyFile.FullName);
 		}
 
 		/// <summary>
-		/// Doc
+		/// Tries to read a sequence of tokens from the given line of text, updating a position within it.
 		/// </summary>
-		public void ReplaceOrAddLine(string StartOfLine, string ReplacementRHS)
+		/// <param name="Line">Line to read tokens from</param>
+		/// <param name="LineOffset">Position within </param>
+		/// <param name="Other">The tokens to attempt to match</param>
+		/// <returns>True if the tokens are matched</returns>
+		static bool MatchCppTokens(string Line, ref int LineOffset, string Other)
 		{
-			if (Contains(StartOfLine))
-			{
-				ReplaceLine(StartOfLine, ReplacementRHS);
-			}
-			else
+			int OtherOffset = 0;
+			for(;;)
 			{
-				AddLine("");
-				AddLine(StartOfLine + ReplacementRHS);
+				// Read the next token from the prefix text. If we've reached the end, return that we've found a match.
+				string OtherToken = ReadCppToken(Other, ref OtherOffset);
+				if(OtherToken == null)
+				{
+					return true;
+				}
+
+				// Otherwise read the next token from the line and check it matches
+				string LineToken = ReadCppToken(Line, ref LineOffset);
+				if(LineToken != OtherToken)
+				{
+					return false;
+				}
 			}
 		}
 
 		/// <summary>
-		/// Adds a new line to the version file
+		/// Reads a single token (using a very rough approximation of a C++ token) from the given line
 		/// </summary>
-		/// <param name="Line"></param>
-		public void AddLine(string Line)
+		/// <param name="Line">The line to read from</param>
+		/// <param name="Offset">Position within the line to start reading. Initial whitespace will be skipped.</param>
+		/// <returns>The token that was read, or null if the end of the string</returns>
+		static string ReadCppToken(string Line, ref int Offset)
 		{
-			Lines.Add(Line);
+			// Skip any leading whitespace
+			while(Offset < Line.Length && Char.IsWhiteSpace(Line[Offset]))
+			{
+				Offset++;
+			}
+			if(Offset == Line.Length)
+			{
+				return null;
+			}
+
+			// Find the length of the token
+			int StartOffset = Offset++;
+			if(Char.IsLetterOrDigit(Line[StartOffset]))
+			{
+				while(Offset < Line.Length)
+				{
+					char Character = Line[Offset];
+					if(!Char.IsLetterOrDigit(Character) && Character != '_')
+					{
+						break;
+					}
+					Offset++;
+				}
+			}
+			return Line.Substring(StartOffset, Offset - StartOffset);
 		}
 
 		/// <summary>
@@ -1133,7 +1184,7 @@ namespace AutomationTool
 
 			if (bDifferent)
 			{
-				MyFile.IsReadOnly = false;
+				MakeFileWriteable(MyFile.FullName);
 				if (!InternalUtils.SafeWriteAllLines(MyFile.FullName, Lines.ToArray()))
 				{
 					throw new AutomationException("Unable to update version info in {0}", MyFile.FullName);
@@ -1143,20 +1194,23 @@ namespace AutomationTool
 		}
 
 		/// <summary>
-		/// Checks if the version file contains the specified string.
+		/// Makes a version file writeable. If P4 is enabled, syncs it to revision 0 to prevent conflicts with the P4 have table rather than checking it out, since we never intend to 
+		/// submit it to Perforce. For existing writeable files, P4V prompts the user to overwrite the next time they sync.
 		/// </summary>
-		/// <param name="Text">String to look for.</param>
-		/// <returns></returns>
-		public bool Contains(string Text, bool CaseSensitive = true)
+		/// <param name="FileName">The file to make writeable</param>
+		public static void MakeFileWriteable(string FileName)
 		{
-			foreach (var Line in Lines)
+			if(CommandUtils.IsReadOnly(FileName))
 			{
-				if (Line.IndexOf(Text, CaseSensitive ? StringComparison.InvariantCulture : StringComparison.InvariantCultureIgnoreCase) >= 0)
+				if(CommandUtils.P4Enabled)
 				{
-					return true;
+					CommandUtils.P4.Sync(String.Format("\"{0}#0\"", FileName), false, false);
+				}
+				if(CommandUtils.FileExists_NoExceptions(FileName) && CommandUtils.IsReadOnly(FileName))
+				{
+					throw new AutomationException("Cannot write to {0}; file is read-only", FileName);
 				}
 			}
-			return false;
 		}
 
 		/// <summary>
diff --git a/Engine/Source/Programs/AutomationTool/BuildGraph/Tasks/SetVersionTask.cs b/Engine/Source/Programs/AutomationTool/BuildGraph/Tasks/SetVersionTask.cs
--- a/Engine/Source/Programs/AutomationTool/BuildGraph/Tasks/SetVersionTask.cs
+++ b/Engine/Source/Programs/AutomationTool/BuildGraph/Tasks/SetVersionTask.cs
@@ -19,6 +19,12 @@ namespace AutomationTool.Tasks
 		[TaskParameter]
 		public int Change;
 
+		/// <summary>
+		/// The engine compatible changelist to set in the version files
+		/// </summary>
+		[TaskParameter(Optional = true)]
+		public int CompatibleChange;
+
 		/// <summary>
 		/// The branch string
 		/// </summary>
@@ -80,7 +86,7 @@ namespace AutomationTool.Tasks
 		public override bool Execute(JobContext Job, HashSet<FileReference> BuildProducts, Dictionary<string, HashSet<FileReference>> TagNameToFileSet)
 		{
 			// Update the version files
-			List<string> FileNames = UE4Build.StaticUpdateVersionFiles(Parameters.Change, Parameters.Branch, Parameters.Build, Parameters.Licensee, !Parameters.SkipWrite);
+			List<string> FileNames = UE4Build.StaticUpdateVersionFiles(Parameters.Change, Parameters.CompatibleChange, Parameters.Branch, Parameters.Build, Parameters.Licensee, !Parameters.SkipWrite);
 			List<FileReference> VersionFiles = FileNames.Select(x => new FileReference(x)).ToList();
 
 			// Apply the optional tag to them
diff --git a/Engine/Source/Programs/AutomationTool/Scripts/Tests.Automation.cs b/Engine/Source/Programs/AutomationTool/Scripts/Tests.Automation.cs
--- a/Engine/Source/Programs/AutomationTool/Scripts/Tests.Automation.cs
+++ b/Engine/Source/Programs/AutomationTool/Scripts/Tests.Automation.cs
@@ -1314,58 +1314,6 @@ public class UBT : BuildCommand
 		Log("UBT Completed");
 	}
 }
-[Help("Zeroes engine versions in ObjectVersion.cpp and Version.h and checks them in.")]
-[RequireP4]
-public class ZeroEngineVersions : BuildCommand
-{
-	public override void ExecuteBuild()
-	{
-		string VersionFilename = CmdEnv.LocalRoot + @"/Engine/Source/Runtime/Launch/Resources/Version.h";
-		if (P4Env.Changelist > 0)
-		{
-			var Stat = P4.FStat(VersionFilename);
-			if (Stat.IsValid && Stat.Action != P4Action.None)
-			{
-				Log("Reverting {0}", VersionFilename);
-				P4.Revert(VersionFilename);
-			}
-
-			Log("Gettting engine version files @{0}", P4Env.Changelist);
-			P4.Sync(String.Format("-f {0}@{1}", VersionFilename, P4Env.Changelist));
-		}
-
-		Log("Checking if engine version files need to be reset...");
-		List<string> FilesToSubmit = new List<string>();
-		{
-			VersionFileUpdater VersionH = new VersionFileUpdater(VersionFilename);
-			if (VersionH.Contains("#define ENGINE_VERSION 0") == false)
-			{
-				Log("Zeroing out engine versions in {0}", VersionFilename);
-				VersionH.ReplaceLine("#define BUILT_FROM_CHANGELIST ", "0");
-
-				VersionH.Commit();
-				FilesToSubmit.Add(VersionFilename);
-			}
-		}
-
-		if (FilesToSubmit.Count > 0)
-		{
-			int CL = P4.CreateChange(null, "Zero engine versions");
-			foreach (var Filename in FilesToSubmit)
-			{
-				P4.Edit(CL, Filename);
-			}
-			Log("Submitting CL #{0}...", CL);
-			int SubmittedCL;
-			P4.Submit(CL, out SubmittedCL, false, true);
-			Log("CL #{0} submitted as {1}", CL, SubmittedCL);
-		}
-		else
-		{
-			Log("Engine version files are set to 0.");
-		}
-	}
-}
 
 [Help("Helper command to sync only source code + engine files from P4.")]
 [Help("Branch", "Optional branch depot path, default is: -Branch=//depot/UE4")]
diff --git a/Engine/Source/Programs/AutomationTool/Scripts/UpdateLocalVersion.Automation.cs b/Engine/Source/Programs/AutomationTool/Scripts/UpdateLocalVersion.Automation.cs
--- a/Engine/Source/Programs/AutomationTool/Scripts/UpdateLocalVersion.Automation.cs
+++ b/Engine/Source/Programs/AutomationTool/Scripts/UpdateLocalVersion.Automation.cs
@@ -15,8 +15,10 @@ public class UpdateLocalVersion : BuildCommand
 {
 	public override void ExecuteBuild()
 	{
-		var UE4Build = new UE4Build(this);
-		var ChangelistOverride = ParseParamInt("cl", -1);
-		UE4Build.UpdateVersionFiles(ChangelistNumberOverride: (ChangelistOverride < 0 ? (int?)null : ChangelistOverride), Build: ParseParamValue("Build", null));
+		UE4Build UE4Build = new UE4Build(this);
+		int? ChangelistOverride = ParseParamNullableInt("cl");
+		int? CompatibleChangelistOverride = ParseParamNullableInt("compatiblecl");
+		string Build = ParseParamValue("Build", null);
+		UE4Build.UpdateVersionFiles(ChangelistNumberOverride: ChangelistOverride, CompatibleChangelistNumberOverride: CompatibleChangelistOverride, Build: Build);
 	}
 }
diff --git a/Engine/Source/Programs/UnrealBuildTool/Configuration/UEBuildTarget.cs b/Engine/Source/Programs/UnrealBuildTool/Configuration/UEBuildTarget.cs
--- a/Engine/Source/Programs/UnrealBuildTool/Configuration/UEBuildTarget.cs
+++ b/Engine/Source/Programs/UnrealBuildTool/Configuration/UEBuildTarget.cs
@@ -1692,6 +1692,33 @@ namespace UnrealBuildTool
 
 				// delete intermediate directory of the extra files
 			}
+
+			// The engine updates the PATH environment variable to supply all valid locations for DLLs, but the Windows loader reads imported DLLs from the first location it finds them. 
+			// If modules are moved from one place to another, we have to be sure to clean up the old versions so that they're not loaded accidentally causing unintuitive import errors.
+			HashSet<FileReference> OutputFiles = new HashSet<FileReference>();
+			Dictionary<string, FileReference> OutputFileNames = new Dictionary<string, FileReference>(StringComparer.InvariantCultureIgnoreCase);
+			foreach(UEBuildBinary Binary in AppBinaries)
+			{
+				foreach(FileReference OutputFile in Binary.Config.OutputFilePaths)
+				{
+					OutputFiles.Add(OutputFile);
+					OutputFileNames[OutputFile.GetFileName()] = OutputFile;
+				}
+			}
+
+			// Search all the output directories for files with a name matching one of our output files
+			foreach(DirectoryReference OutputDirectory in OutputFiles.Select(x => x.Directory).Distinct())
+			{
+				foreach(FileReference ExistingFile in OutputDirectory.EnumerateFileReferences())
+				{
+					FileReference OutputFile;
+					if(OutputFileNames.TryGetValue(ExistingFile.GetFileName(), out OutputFile) && !OutputFiles.Contains(ExistingFile))
+					{
+						Log.TraceInformation("Deleting '{0}' to avoid ambiguity with '{1}'", ExistingFile, OutputFile);
+						CleanFile(ExistingFile.FullName);
+					}
+				}
+			}
 		}
 
 		/// <summary>
@@ -2228,16 +2255,17 @@ namespace UnrealBuildTool
 		/// <summary>
 		/// Builds the target, appending list of output files and returns building result.
 		/// </summary>
-		public ECompilationResult Build(UEToolChain TargetToolChain, out List<FileItem> OutputItems, out List<UHTModuleInfo> UObjectModules, out string EULAViolationWarning)
+		public ECompilationResult Build(UEToolChain TargetToolChain, out List<FileItem> OutputItems, out List<UHTModuleInfo> UObjectModules)
 		{
 			OutputItems = new List<FileItem>();
 			UObjectModules = new List<UHTModuleInfo>();
 
 			PreBuildSetup(TargetToolChain);
 
-			EULAViolationWarning = !ProjectFileGenerator.bGenerateProjectFiles
-				? CheckForEULAViolation()
-				: null;
+			if(!ProjectFileGenerator.bGenerateProjectFiles)
+			{
+				CheckForEULAViolation();
+			}
 
 			// Execute the pre-build steps
 			if(!ExecuteCustomPreBuildSteps())
@@ -2400,15 +2428,12 @@ namespace UnrealBuildTool
 		/// <summary>
 		/// Check for EULA violation dependency issues.
 		/// </summary>
-		/// <returns>EULAViolationWarning if error hasn't been thrown earlier. This behavior is controlled using BuildConfiguration.bBreakBuildOnLicenseViolation.</returns>
-		private string CheckForEULAViolation()
+		private void CheckForEULAViolation()
 		{
-			string EULAViolationWarning = null;
-
 			if (TargetType != TargetRules.TargetType.Editor && TargetType != TargetRules.TargetType.Program && Configuration == UnrealTargetConfiguration.Shipping &&
 				BuildConfiguration.bCheckLicenseViolations)
 			{
-				StringBuilder RedistributionErrorMessageBuilder = new StringBuilder();
+				bool bLicenseViolation = false;
 				foreach (UEBuildBinary Binary in AppBinaries)
 				{
 					IEnumerable<UEBuildModule> NonRedistModules = Binary.GetAllDependencyModules(true, false).Where((DependencyModule) =>
@@ -2417,25 +2442,23 @@ namespace UnrealBuildTool
 
 					if (NonRedistModules.Count() != 0)
 					{
-						RedistributionErrorMessageBuilder.Append(string.Format(
-							"{2}Binary {0} depends on: {1}.", Binary.ToString(), string.Join(", ", NonRedistModules),
-							RedistributionErrorMessageBuilder.Length > 0 ? "\n" : ""
-						));
+						string Message = string.Format("Non-editor build cannot depend on non-redistributable modules. {0} depends on '{1}'.", Binary.ToString(), string.Join("', '", NonRedistModules));
+						if(BuildConfiguration.bBreakBuildOnLicenseViolation)
+						{
+							Log.TraceError("ERROR: {0}", Message);
+						}
+						else
+						{
+							Log.TraceWarning("WARNING: {0}", Message);
+						}
+						bLicenseViolation = true;
 					}
 				}
-
-				if (RedistributionErrorMessageBuilder.Length > 0)
+				if (BuildConfiguration.bBreakBuildOnLicenseViolation && bLicenseViolation)
 				{
-					EULAViolationWarning = string.Format("Non-editor build cannot depend on non-redistributable modules. Details:\n{0}", RedistributionErrorMessageBuilder.ToString());
-
-					if (BuildConfiguration.bBreakBuildOnLicenseViolation)
-					{
-						throw new BuildException(EULAViolationWarning);
-					}
+					throw new BuildException("Non-editor build cannot depend on non-redistributable modules.");
 				}
 			}
-
-			return EULAViolationWarning;
 		}
 
 		/// <summary>
diff --git a/Engine/Source/Programs/UnrealBuildTool/Mac/MacToolChain.cs b/Engine/Source/Programs/UnrealBuildTool/Mac/MacToolChain.cs
--- a/Engine/Source/Programs/UnrealBuildTool/Mac/MacToolChain.cs
+++ b/Engine/Source/Programs/UnrealBuildTool/Mac/MacToolChain.cs
@@ -12,7 +12,7 @@ using Ionic.Zip;
 
 namespace UnrealBuildTool
 {
-	class MacToolChain : AppleToolChain
+	public class MacToolChain : AppleToolChain
 	{
 		public MacToolChain(FileReference InProjectFile) 
 			: base(CPPTargetPlatform.Mac, UnrealTargetPlatform.Mac, InProjectFile)
@@ -471,7 +471,7 @@ namespace UnrealBuildTool
 		private int LoadEngineCL()
 		{
 			BuildVersion Version;
-			if (BuildVersion.TryRead("../Build/Build.version", out Version))
+			if (BuildVersion.TryRead(out Version))
 			{
 				return Version.Changelist;
 			}
@@ -483,26 +483,15 @@ namespace UnrealBuildTool
 
 		public static string LoadEngineDisplayVersion(bool bIgnorePatchVersion = false)
 		{
-			string[] VersionHeader = Utils.ReadAllText("../Source/Runtime/Launch/Resources/Version.h").Replace("\r\n", "\n").Replace("\t", " ").Split('\n');
-			string EngineVersionMajor = "4";
-			string EngineVersionMinor = "0";
-			string EngineVersionPatch = "0";
-			foreach (string Line in VersionHeader)
+			BuildVersion Version;
+			if (BuildVersion.TryRead(out Version))
 			{
-				if (Line.StartsWith("#define ENGINE_MAJOR_VERSION "))
-				{
-					EngineVersionMajor = Line.Split(' ')[2];
-				}
-				else if (Line.StartsWith("#define ENGINE_MINOR_VERSION "))
-				{
-					EngineVersionMinor = Line.Split(' ')[2];
-				}
-				else if (Line.StartsWith("#define ENGINE_PATCH_VERSION ") && !bIgnorePatchVersion)
-				{
-					EngineVersionPatch = Line.Split(' ')[2];
-				}
+				return String.Format("{0}.{1}.{2}", Version.MajorVersion, Version.MinorVersion, bIgnorePatchVersion? 0 : Version.PatchVersion);
+			}
+			else
+			{
+				return "4.0.0";
 			}
-			return EngineVersionMajor + "." + EngineVersionMinor + "." + EngineVersionPatch;
 		}
 
 		private string LoadLauncherDisplayVersion()
@@ -531,59 +520,32 @@ namespace UnrealBuildTool
 
 		private int LoadBuiltFromChangelistValue()
 		{
-			string[] VersionHeader = Utils.ReadAllText("../Source/Runtime/Launch/Resources/Version.h").Replace("\r\n", "\n").Replace("\t", " ").Split('\n');
-			foreach (string Line in VersionHeader)
-			{
-				if (Line.StartsWith("#define BUILT_FROM_CHANGELIST "))
-				{
-					return int.Parse(Line.Split(' ')[2]);
-				}
-			}
-			return 0;
+			return LoadEngineCL();
 		}
 
 		private int LoadIsLicenseeVersionValue()
 		{
-			string[] VersionHeader = Utils.ReadAllText("../Source/Runtime/Launch/Resources/Version.h").Replace("\r\n", "\n").Replace("\t", " ").Split('\n');
-			foreach (string Line in VersionHeader)
+			BuildVersion Version;
+			if (BuildVersion.TryRead(out Version))
 			{
-				if (Line.StartsWith("#define ENGINE_IS_LICENSEE_VERSION "))
-				{
-					return int.Parse(Line.Split(' ')[2]);
-				}
+				return Version.IsLicenseeVersion;
+			}
+			else
+			{
+				return 0;
 			}
-			return 0;
 		}
 
 		private string LoadEngineAPIVersion()
 		{
 			int CL = 0;
-			// @todo: Temp solution to work around a problem with parsing ModuleVersion.h updated for 4.4.1 hotfix
-			int BuiltFromChangelist = LoadBuiltFromChangelistValue();
-			if (BuiltFromChangelist > 0)
+
+			BuildVersion Version;
+			if (BuildVersion.TryRead(out Version))
 			{
-				bool bFirstApiVersionDefine = true;
-				foreach (string Line in File.ReadAllLines("../Source/Runtime/Core/Public/Modules/ModuleVersion.h"))
-				{
-					string[] Tokens = Line.Split(new char[] { ' ', '\t' }, StringSplitOptions.RemoveEmptyEntries);
-					if (Tokens.Length >= 3 && Tokens[0] == "#define" && Tokens[1] == "MODULE_API_VERSION")
-					{
-						if (!bFirstApiVersionDefine || LoadIsLicenseeVersionValue() != 0)
-						{
-							if (Tokens[2] == "BUILT_FROM_CHANGELIST")
-							{
-								CL = LoadEngineCL();
-							}
-							else
-							{
-								CL = int.Parse(Tokens[2]);
-							}
-							break;
-						}
-						bFirstApiVersionDefine = false;
-					}
-				}
+				CL = (Version.CompatibleChangelist != 0)? Version.CompatibleChangelist : Version.Changelist;
 			}
+
 			return String.Format("{0}.{1}.{2}", CL / (100 * 100), (CL / 100) % 100, CL % 100);
 		}
 
diff --git a/Engine/Source/Programs/UnrealBuildTool/System/ExternalExecution.cs b/Engine/Source/Programs/UnrealBuildTool/System/ExternalExecution.cs
--- a/Engine/Source/Programs/UnrealBuildTool/System/ExternalExecution.cs
+++ b/Engine/Source/Programs/UnrealBuildTool/System/ExternalExecution.cs
@@ -1001,7 +1001,7 @@ namespace UnrealBuildTool
 								UEBuildModuleCPP DependencyModuleCPP = (UEBuildModuleCPP)Target.GetModuleByName(UHTModuleInfo.ModuleName);
 								CPPEnvironment ModuleCompileEnvironment = DependencyModuleCPP.CreateModuleCompileEnvironment(Target, GlobalCompileEnvironment);
 								DependencyModuleCPP.CachePCHUsageForModuleSourceFiles(Target, ModuleCompileEnvironment);
-								if (DependencyModuleCPP.ProcessedDependencies.UniquePCHHeaderFile != null)
+								if (DependencyModuleCPP.ProcessedDependencies != null && DependencyModuleCPP.ProcessedDependencies.UniquePCHHeaderFile != null)
 								{
 									UHTModuleInfo.PCH = DependencyModuleCPP.ProcessedDependencies.UniquePCHHeaderFile.AbsolutePath;
 								}
diff --git a/Engine/Source/Programs/UnrealBuildTool/System/FileSystemReference.cs b/Engine/Source/Programs/UnrealBuildTool/System/FileSystemReference.cs
--- a/Engine/Source/Programs/UnrealBuildTool/System/FileSystemReference.cs
+++ b/Engine/Source/Programs/UnrealBuildTool/System/FileSystemReference.cs
@@ -26,21 +26,12 @@ namespace UnrealBuildTool
 		public readonly string CanonicalName;
 
 		/// <summary>
-		/// Constructs a filesystem object for the given path.
-		/// </summary>
-		public FileSystemReference(string InPath)
-		{
-			FullName = Path.GetFullPath(InPath).TrimEnd(Path.DirectorySeparatorChar);
-			CanonicalName = FullName.ToLowerInvariant();
-		}
-
-		/// <summary>
-		/// Constructs a reference from the given FileSystemInfo.
+		/// Direct constructor for a path
 		/// </summary>
-		public FileSystemReference(FileSystemInfo InInfo)
+		protected FileSystemReference(string InFullName)
 		{
-			FullName = InInfo.FullName;
-			CanonicalName = FullName.ToLowerInvariant();
+			FullName = InFullName;
+			CanonicalName = InFullName.ToLowerInvariant();
 		}
 
 		/// <summary>
@@ -254,7 +245,7 @@ namespace UnrealBuildTool
 		/// </summary>
 		/// <param name="InPath">Path to this directory.</param>
 		public DirectoryReference(string InPath)
-			: base(InPath)
+			: base(FixTrailingPathSeparator(Path.GetFullPath(InPath)))
 		{
 		}
 
@@ -263,7 +254,7 @@ namespace UnrealBuildTool
 		/// </summary>
 		/// <param name="InInfo">Path to this file</param>
 		public DirectoryReference(DirectoryInfo InInfo)
-			: base(InInfo)
+			: base(FixTrailingPathSeparator(InInfo.FullName))
 		{
 		}
 
@@ -277,6 +268,32 @@ namespace UnrealBuildTool
 		{
 		}
 
+		/// <summary>
+		/// Ensures that the correct trailing path separator is appended. On Windows, the root directory (eg. C:\) always has a trailing path separator, but no other
+		/// path does.
+		/// </summary>
+		/// <param name="DirName">Absolute path to the directory</param>
+		/// <returns>Path to the directory, with the correct trailing path separator</returns>
+		private static string FixTrailingPathSeparator(string DirName)
+		{
+			if(DirName.Length == 2 && DirName[2] == ':')
+			{
+				return DirName + Path.DirectorySeparatorChar;
+			}
+			else if(DirName.Length == 3 && DirName[1] == ':' && DirName[2] == Path.DirectorySeparatorChar)
+			{
+				return DirName;
+			}
+			else if(DirName.Length > 1 && DirName[DirName.Length - 1] == Path.DirectorySeparatorChar)
+			{
+				return DirName.TrimEnd(Path.DirectorySeparatorChar);
+			}
+			else
+			{
+				return DirName;
+			}
+		}
+
 		/// <summary>
 		/// Gets the top level directory name
 		/// </summary>
@@ -530,8 +547,12 @@ namespace UnrealBuildTool
 		/// </summary>
 		/// <param name="InPath">Path to this file</param>
 		public FileReference(string InPath)
-			: base(InPath)
+			: base(Path.GetFullPath(InPath))
 		{
+			if(FullName.EndsWith("\\") || FullName.EndsWith("/"))
+			{
+				throw new ArgumentException("File names may not be terminated by a path separator character");
+			}
 		}
 
 		/// <summary>
@@ -539,7 +560,7 @@ namespace UnrealBuildTool
 		/// </summary>
 		/// <param name="InInfo">Path to this file</param>
 		public FileReference(FileInfo InInfo)
-			: base(InInfo)
+			: base(InInfo.FullName)
 		{
 		}
 
diff --git a/Engine/Source/Programs/UnrealBuildTool/System/TargetReceipt.cs b/Engine/Source/Programs/UnrealBuildTool/System/TargetReceipt.cs
--- a/Engine/Source/Programs/UnrealBuildTool/System/TargetReceipt.cs
+++ b/Engine/Source/Programs/UnrealBuildTool/System/TargetReceipt.cs
@@ -627,7 +627,7 @@ namespace UnrealBuildTool
 				Writer.WriteValue("BuildId", BuildId);
 
 				Writer.WriteObjectStart("Version");
-				Version.Write(Writer);
+				Version.WriteProperties(Writer);
 				Writer.WriteObjectEnd();
 
 				Writer.WriteArrayStart("BuildProducts");
diff --git a/Engine/Source/Programs/UnrealBuildTool/System/UnrealBuildTool.cs b/Engine/Source/Programs/UnrealBuildTool/System/UnrealBuildTool.cs
--- a/Engine/Source/Programs/UnrealBuildTool/System/UnrealBuildTool.cs
+++ b/Engine/Source/Programs/UnrealBuildTool/System/UnrealBuildTool.cs
@@ -1584,7 +1584,6 @@ namespace UnrealBuildTool
 
 			UEToolChain ToolChain = BuildPlatformContext.CreateToolChainForDefaultCppPlatform();
 
-			string EULAViolationWarning = null;
 			Thread CPPIncludesThread = null;
 
 			try
@@ -1844,7 +1843,7 @@ namespace UnrealBuildTool
 					{
 						List<FileItem> TargetOutputItems;
 						List<UHTModuleInfo> TargetUObjectModules;
-						BuildResult = Target.Build(ToolChain, out TargetOutputItems, out TargetUObjectModules, out EULAViolationWarning);
+						BuildResult = Target.Build(ToolChain, out TargetOutputItems, out TargetUObjectModules);
 						if (BuildResult != ECompilationResult.Succeeded)
 						{
 							break;
@@ -2097,11 +2096,6 @@ namespace UnrealBuildTool
 				CPPEnvironment.FlatCPPIncludeDependencyCache.Clear();
 			}
 
-			if (EULAViolationWarning != null)
-			{
-				Log.TraceWarning("WARNING: {0}", EULAViolationWarning);
-			}
-
 			// Figure out how long we took to execute.
 			double BuildDuration = (DateTime.UtcNow - StartTime).TotalSeconds;
 			if (ExecutorName == "Local" || ExecutorName == "Distcc" || ExecutorName == "SNDBS")
diff --git a/Engine/Source/Programs/UnrealBuildTool/System/Utils.cs b/Engine/Source/Programs/UnrealBuildTool/System/Utils.cs
--- a/Engine/Source/Programs/UnrealBuildTool/System/Utils.cs
+++ b/Engine/Source/Programs/UnrealBuildTool/System/Utils.cs
@@ -27,11 +27,22 @@ namespace UnrealBuildTool
 		public int MinorVersion;
 		public int PatchVersion;
 		public int Changelist;
+		public int CompatibleChangelist;
 		public int IsLicenseeVersion;
 		public string BranchName;
 
 		/// <summary>
-		/// Try to read the Build/Build.version file from disk
+		/// Try to read a version file from disk
+		/// </summary>
+		/// <param name="Version">The version information</param>
+		/// <returns>True if the version was read sucessfully, false otherwise</returns>
+		public static bool TryRead(out BuildVersion Version)
+		{
+			return TryRead(GetDefaultFileName(), out Version);
+		}
+
+		/// <summary>
+		/// Try to read a version file from disk
 		/// </summary>
 		/// <param name="Version">The version information</param>
 		/// <returns>True if the version was read sucessfully, false otherwise</returns>
@@ -46,6 +57,15 @@ namespace UnrealBuildTool
 			return TryParse(Object, out Version);
 		}
 
+		/// <summary>
+		/// Get the default path to the build.version file on disk
+		/// </summary>
+		/// <returns>Path to the Build.version file</returns>
+		public static string GetDefaultFileName()
+		{
+			return FileReference.Combine(UnrealBuildTool.EngineDirectory, "Build", "Build.version").FullName;
+		}
+
 		/// <summary>
 		/// Parses a build version from a JsonObject
 		/// </summary>
@@ -62,6 +82,13 @@ namespace UnrealBuildTool
 			}
 
 			Object.TryGetIntegerField("Changelist", out NewVersion.Changelist);
+			if(NewVersion.Changelist != 0)
+			{
+				if(!Object.TryGetIntegerField("CompatibleChangelist", out NewVersion.CompatibleChangelist))
+				{
+					NewVersion.CompatibleChangelist = NewVersion.Changelist;
+				}
+			}
 			Object.TryGetIntegerField("IsLicenseeVersion", out NewVersion.IsLicenseeVersion);
 			Object.TryGetStringField("BranchName", out NewVersion.BranchName);
 
@@ -75,12 +102,29 @@ namespace UnrealBuildTool
 		/// <param name="Object">The object to read from</param>
 		/// <param name="Version">The resulting version field</param>
 		/// <returns>True if the build version could be read, false otherwise</returns>
-		public void Write(JsonWriter Writer)
+		public void Write(string FileName)
+		{
+			using (JsonWriter Writer = new JsonWriter(FileName))
+			{
+				Writer.WriteObjectStart();
+				WriteProperties(Writer);
+				Writer.WriteObjectEnd();
+			}
+		}
+
+		/// <summary>
+		/// Exports this object as Json
+		/// </summary>
+		/// <param name="Object">The object to read from</param>
+		/// <param name="Version">The resulting version field</param>
+		/// <returns>True if the build version could be read, false otherwise</returns>
+		public void WriteProperties(JsonWriter Writer)
 		{
 			Writer.WriteValue("MajorVersion", MajorVersion);
 			Writer.WriteValue("MinorVersion", MinorVersion);
 			Writer.WriteValue("PatchVersion", PatchVersion);
 			Writer.WriteValue("Changelist", Changelist);
+			Writer.WriteValue("CompatibleChangelist", CompatibleChangelist);
 			Writer.WriteValue("IsLicenseeVersion", IsLicenseeVersion);
 			Writer.WriteValue("BranchName", BranchName);
 		}
diff --git a/Engine/Source/Programs/UnrealBuildTool/System/VCProjectFileGenerator.cs b/Engine/Source/Programs/UnrealBuildTool/System/VCProjectFileGenerator.cs
--- a/Engine/Source/Programs/UnrealBuildTool/System/VCProjectFileGenerator.cs
+++ b/Engine/Source/Programs/UnrealBuildTool/System/VCProjectFileGenerator.cs
@@ -651,7 +651,7 @@ namespace UnrealBuildTool
 										!IsProgramProject && SolutionConfigCombination.TargetConfigurationName != TargetRules.TargetType.Program.ToString())
 									{
 										string TargetConfigurationName = SolutionConfigCombination.TargetConfigurationName;
-										if (IsProgramProject && TargetConfigurationName == TargetRules.TargetType.Game.ToString())
+										if (IsProgramProject)
 										{
 											TargetConfigurationName = TargetRules.TargetType.Program.ToString();
 										}
diff --git a/Engine/Source/Programs/UnrealCodeAnalyzer/UnrealCodeAnalyzer.Target.cs b/Engine/Source/Programs/UnrealCodeAnalyzer/UnrealCodeAnalyzer.Target.cs
--- a/Engine/Source/Programs/UnrealCodeAnalyzer/UnrealCodeAnalyzer.Target.cs
+++ b/Engine/Source/Programs/UnrealCodeAnalyzer/UnrealCodeAnalyzer.Target.cs
@@ -90,7 +90,6 @@ public class UnrealCodeAnalyzerTarget : TargetRules
 
 	public override bool GetSupportedPlatforms(ref List<UnrealTargetPlatform> OutPlatforms)
 	{
-		OutPlatforms.Clear();
 		OutPlatforms.Add(UnrealTargetPlatform.Win32);
 		return true;
 	}
diff --git a/Engine/Source/Programs/UnrealGameSync/UnrealGameSync/FindFoldersToCleanTask.cs b/Engine/Source/Programs/UnrealGameSync/UnrealGameSync/FindFoldersToCleanTask.cs
--- a/Engine/Source/Programs/UnrealGameSync/UnrealGameSync/FindFoldersToCleanTask.cs
+++ b/Engine/Source/Programs/UnrealGameSync/UnrealGameSync/FindFoldersToCleanTask.cs
@@ -16,6 +16,7 @@ namespace UnrealGameSync
 		public DirectoryInfo Directory;
 		public List<FolderToClean> SubFolders = new List<FolderToClean>();
 		public List<FileInfo> FilesToClean = new List<FileInfo>();
+		public bool bEmptyLeaf = false;
 		public bool bEmptyAfterDelete = true;
 
 		public FolderToClean(DirectoryInfo InDirectory)
@@ -115,6 +116,7 @@ namespace UnrealGameSync
 					QueueFolderToPopulate(SubFolder);
 				}
 				Folder.FilesToClean = Folder.Directory.EnumerateFiles().ToList();
+				Folder.bEmptyLeaf = Folder.SubFolders.Count == 0 && Folder.FilesToClean.Count == 0;
 			}
 
 			if(Interlocked.Decrement(ref RemainingFoldersToScan) == 0)
@@ -143,7 +145,7 @@ namespace UnrealGameSync
 			{
 				RemoveEmptyFolders(SubFolder);
 			}
-			Folder.SubFolders.RemoveAll(x => x.SubFolders.Count == 0 && x.FilesToClean.Count == 0);
+			Folder.SubFolders.RemoveAll(x => x.SubFolders.Count == 0 && x.FilesToClean.Count == 0 && !x.bEmptyLeaf);
 		}
 
 		public void Dispose()
diff --git a/Engine/Source/Programs/UnrealGameSync/UnrealGameSync/Forms/CleanWorkspaceWindow.cs b/Engine/Source/Programs/UnrealGameSync/UnrealGameSync/Forms/CleanWorkspaceWindow.cs
--- a/Engine/Source/Programs/UnrealGameSync/UnrealGameSync/Forms/CleanWorkspaceWindow.cs
+++ b/Engine/Source/Programs/UnrealGameSync/UnrealGameSync/Forms/CleanWorkspaceWindow.cs
@@ -25,6 +25,16 @@ namespace UnrealGameSync
 			public FolderToClean Folder;
 			public int NumFiles;
 			public int NumSelectedFiles;
+			public int NumEmptySelectedFiles;
+			public int NumDefaultSelectedFiles;
+		}
+
+		enum SelectionType
+		{
+			All,
+			SafeToDelete,
+			Empty,
+			None,
 		}
 
 		static readonly CheckBoxState[] CheckBoxStates = 
@@ -155,7 +165,7 @@ namespace UnrealGameSync
 
 		private TreeNode BuildTreeViewStructure(FolderToClean Folder, string FolderPath, bool bParentFolderSelected, int Depth)
 		{
-			bool bSelectFolder = bParentFolderSelected || IsSafeToDeleteFolder(FolderPath);
+			bool bSelectFolder = bParentFolderSelected || IsSafeToDeleteFolder(FolderPath) || Folder.bEmptyLeaf;
 
 			TreeNodeData FolderNodeData = new TreeNodeData();
 			FolderNodeData.Folder = Folder;
@@ -172,6 +182,8 @@ namespace UnrealGameSync
 				TreeNodeData ChildNodeData = (TreeNodeData)ChildNode.Tag;
 				FolderNodeData.NumFiles += ChildNodeData.NumFiles;
 				FolderNodeData.NumSelectedFiles += ChildNodeData.NumSelectedFiles;
+				FolderNodeData.NumEmptySelectedFiles += ChildNodeData.NumEmptySelectedFiles;
+				FolderNodeData.NumDefaultSelectedFiles += ChildNodeData.NumDefaultSelectedFiles;
 			}
 
 			foreach(FileInfo File in Folder.FilesToClean.OrderBy(x => x.Name))
@@ -184,6 +196,8 @@ namespace UnrealGameSync
 				FileNodeData.File = File;
 				FileNodeData.NumFiles = 1;
 				FileNodeData.NumSelectedFiles = bSelectFile? 1 : 0;
+				FileNodeData.NumEmptySelectedFiles = 0;
+				FileNodeData.NumDefaultSelectedFiles = FileNodeData.NumSelectedFiles;
 
 				TreeNode FileNode = new TreeNode();
 				FileNode.Text = File.Name;
@@ -194,6 +208,16 @@ namespace UnrealGameSync
 
 				FolderNodeData.NumFiles++;
 				FolderNodeData.NumSelectedFiles += FileNodeData.NumSelectedFiles;
+				FolderNodeData.NumEmptySelectedFiles += FileNodeData.NumEmptySelectedFiles;
+				FolderNodeData.NumDefaultSelectedFiles += FileNodeData.NumDefaultSelectedFiles;
+			}
+
+			if(FolderNodeData.Folder.bEmptyLeaf)
+			{
+				FolderNodeData.NumFiles++;
+				FolderNodeData.NumSelectedFiles++;
+				FolderNodeData.NumEmptySelectedFiles++;
+				FolderNodeData.NumDefaultSelectedFiles++;
 			}
 
 			if(FolderNodeData.NumSelectedFiles > 0 && !FolderNodeData.Folder.bEmptyAfterDelete && Depth < 2)
@@ -226,20 +250,26 @@ namespace UnrealGameSync
 
 		private void TreeView_NodeMouseClick(object sender, TreeNodeMouseClickEventArgs e)
 		{
-			if(e.X >= e.Node.Bounds.Left - 32 && e.X < e.Node.Bounds.Left - 16)
+			TreeNode Node = e.Node;
+			if(e.Button == System.Windows.Forms.MouseButtons.Right)
+			{
+				TreeView.SelectedNode = e.Node;
+				FolderContextMenu.Tag = e.Node;
+				FolderContextMenu.Show(TreeView.PointToScreen(e.Location));
+			}
+			else if(e.X >= e.Node.Bounds.Left - 32 && e.X < e.Node.Bounds.Left - 16)
 			{
-				TreeNode Node = e.Node;
 				TreeNodeData NodeData = (TreeNodeData)Node.Tag;
-				SetSelected(Node, NodeData.NumSelectedFiles == 0);
+				SetSelected(Node, (NodeData.NumSelectedFiles == 0)? SelectionType.All : SelectionType.None);
 			}
 		}
 
-		private void SetSelected(TreeNode ParentNode, bool bSelected)
+		private void SetSelected(TreeNode ParentNode, SelectionType Type)
 		{
 			TreeNodeData ParentNodeData = (TreeNodeData)ParentNode.Tag;
 
 			int PrevNumSelectedFiles = ParentNodeData.NumSelectedFiles;
-			SetSelectedOnChildren(ParentNode, bSelected);
+			SetSelectedOnChildren(ParentNode, Type);
 
 			int DeltaNumSelectedFiles = ParentNodeData.NumSelectedFiles - PrevNumSelectedFiles;
 			if(DeltaNumSelectedFiles != 0)
@@ -253,16 +283,32 @@ namespace UnrealGameSync
 			}
 		}
 
-		private void SetSelectedOnChildren(TreeNode ParentNode, bool bSelected)
+		private void SetSelectedOnChildren(TreeNode ParentNode, SelectionType Type)
 		{
 			TreeNodeData ParentNodeData = (TreeNodeData)ParentNode.Tag;
 
-			int NewNumSelectedFiles = bSelected? ParentNodeData.NumFiles : 0;
+			int NewNumSelectedFiles = 0;
+			switch(Type)
+			{
+				case SelectionType.All:
+					NewNumSelectedFiles = ParentNodeData.NumFiles;
+					break;
+				case SelectionType.Empty:
+					NewNumSelectedFiles = ParentNodeData.NumEmptySelectedFiles;
+					break;
+				case SelectionType.SafeToDelete:
+					NewNumSelectedFiles = ParentNodeData.NumDefaultSelectedFiles;
+					break;
+				case SelectionType.None:
+					NewNumSelectedFiles = 0;
+					break;
+			}
+
 			if(NewNumSelectedFiles != ParentNodeData.NumSelectedFiles)
 			{
 				foreach(TreeNode ChildNode in ParentNode.Nodes)
 				{
-					SetSelectedOnChildren(ChildNode, bSelected);
+					SetSelectedOnChildren(ChildNode, Type);
 				}
 				ParentNodeData.NumSelectedFiles = NewNumSelectedFiles;
 				UpdateImage(ParentNode);
@@ -325,8 +371,39 @@ namespace UnrealGameSync
 		{
 			foreach(TreeNode Node in TreeView.Nodes)
 			{
-				SetSelected(Node, true);
+				SetSelected(Node, SelectionType.All);
 			}
 		}
+
+		private void FolderContextMenu_SelectAll_Click(object sender, EventArgs e)
+		{
+			TreeNode Node = (TreeNode)FolderContextMenu.Tag;
+			SetSelected(Node, SelectionType.All);
+		}
+
+		private void FolderContextMenu_SelectSafeToDelete_Click(object sender, EventArgs e)
+		{
+			TreeNode Node = (TreeNode)FolderContextMenu.Tag;
+			SetSelected(Node, SelectionType.SafeToDelete);
+		}
+
+		private void FolderContextMenu_SelectEmptyFolder_Click(object sender, EventArgs e)
+		{
+			TreeNode Node = (TreeNode)FolderContextMenu.Tag;
+			SetSelected(Node, SelectionType.Empty);
+		}
+
+		private void FolderContextMenu_SelectNone_Click(object sender, EventArgs e)
+		{
+			TreeNode Node = (TreeNode)FolderContextMenu.Tag;
+			SetSelected(Node, SelectionType.None);
+		}
+
+		private void FolderContextMenu_OpenWithExplorer_Click(object sender, EventArgs e)
+		{
+			TreeNode Node = (TreeNode)FolderContextMenu.Tag;
+			TreeNodeData NodeData = (TreeNodeData)Node.Tag;
+			Process.Start("explorer.exe", String.Format("\"{0}\"", NodeData.Folder.Directory.FullName));
+		}
 	}
 }
diff --git a/Engine/Source/Programs/UnrealGameSync/UnrealGameSync/Forms/MainWindow.cs b/Engine/Source/Programs/UnrealGameSync/UnrealGameSync/Forms/MainWindow.cs
--- a/Engine/Source/Programs/UnrealGameSync/UnrealGameSync/Forms/MainWindow.cs
+++ b/Engine/Source/Programs/UnrealGameSync/UnrealGameSync/Forms/MainWindow.cs
@@ -1007,7 +1007,7 @@ namespace UnrealGameSync
 				List<BuildData> NotifyBuilds = new List<BuildData>();
 				foreach(BuildData LastBuild in TypeToLastBuild.Values.OrderBy(x => x.BuildType))
 				{
-					if(LastBuild.Result == BuildDataResult.Failure || LastBuild.Result == BuildDataResult.Warning)
+					if((LastBuild.Result == BuildDataResult.Failure || LastBuild.Result == BuildDataResult.Warning) && LastBuild.ChangeNumber >= LastCodeChangeByCurrentUser)
 					{
 						BuildData LastSuccessfulBuild;
 						if(!TypeToLastSucceededBuild.TryGetValue(LastBuild.BuildType, out LastSuccessfulBuild) || LastSuccessfulBuild.ChangeNumber < LastCodeChangeByCurrentUser)
diff --git a/Engine/Source/Programs/UnrealGameSync/UnrealGameSync/Properties/AssemblyInfo.cs b/Engine/Source/Programs/UnrealGameSync/UnrealGameSync/Properties/AssemblyInfo.cs
--- a/Engine/Source/Programs/UnrealGameSync/UnrealGameSync/Properties/AssemblyInfo.cs
+++ b/Engine/Source/Programs/UnrealGameSync/UnrealGameSync/Properties/AssemblyInfo.cs
@@ -34,5 +34,5 @@ using System.Runtime.InteropServices;
 // You can specify all the values or you can default the Build and Revision Numbers 
 // by using the '*' as shown below:
 // [assembly: AssemblyVersion("1.0.*")]
-[assembly: AssemblyVersion("1.76.*")]
-[assembly: AssemblyFileVersion("1.76.0.0")]
+[assembly: AssemblyVersion("1.78.*")]
+[assembly: AssemblyFileVersion("1.78.0.0")]
diff --git a/Engine/Source/Programs/UnrealHeaderTool/Private/Classes.h b/Engine/Source/Programs/UnrealHeaderTool/Private/Classes.h
--- a/Engine/Source/Programs/UnrealHeaderTool/Private/Classes.h
+++ b/Engine/Source/Programs/UnrealHeaderTool/Private/Classes.h
@@ -2,6 +2,7 @@
 
 #pragma once
 
+#include "ParserClass.h"
 #include "ClassTree.h"
 
 #define WIP_UHT_REFACTOR 1
diff --git a/Engine/Source/Programs/UnrealHeaderTool/Private/CodeGenerator.cpp b/Engine/Source/Programs/UnrealHeaderTool/Private/CodeGenerator.cpp
--- a/Engine/Source/Programs/UnrealHeaderTool/Private/CodeGenerator.cpp
+++ b/Engine/Source/Programs/UnrealHeaderTool/Private/CodeGenerator.cpp
@@ -2645,8 +2645,6 @@ void FNativeClassHeaderGenerator::ExportClassesFromSourceFileWrapper(FUnrealSour
 		TEXT("\tThis is automatically generated by the tools.\r\n")
 		TEXT("\tDO NOT modify this manually! Edit the corresponding .h files instead!\r\n")
 		TEXT("===========================================================================*/\r\n")
-		LINE_TERMINATOR
-		TEXT("#include \"ObjectBase.h\"") LINE_TERMINATOR
 		LINE_TERMINATOR);
 
 	GeneratedHeaderTextWithCopyright.Log(TEXT("PRAGMA_DISABLE_DEPRECATION_WARNINGS") LINE_TERMINATOR);
diff --git a/Engine/Source/Runtime/Analytics/AnalyticsVisualEditing/Private/AnalyticsVisualEditing.cpp b/Engine/Source/Runtime/Analytics/AnalyticsVisualEditing/Private/AnalyticsVisualEditing.cpp
--- a/Engine/Source/Runtime/Analytics/AnalyticsVisualEditing/Private/AnalyticsVisualEditing.cpp
+++ b/Engine/Source/Runtime/Analytics/AnalyticsVisualEditing/Private/AnalyticsVisualEditing.cpp
@@ -1,6 +1,7 @@
 // Copyright 1998-2016 Epic Games, Inc. All Rights Reserved.
 
 #include "Analytics.h"
+#include "CoreUObject.h"
 #include "Engine/DeveloperSettings.h"
 #include "AnalyticsVisualEditing.h"
 #include "AnalyticsSettings.h"
diff --git a/Engine/Source/Runtime/Core/Private/GenericPlatform/GenericPlatformChunkInstall.cpp b/Engine/Source/Runtime/Core/Private/GenericPlatform/GenericPlatformChunkInstall.cpp
new file mode 100644
--- /dev/null
+++ b/Engine/Source/Runtime/Core/Private/GenericPlatform/GenericPlatformChunkInstall.cpp
@@ -0,0 +1,6 @@
+// Copyright 1998-2016 Epic Games, Inc. All Rights Reserved.
+
+#include "CorePrivatePCH.h"
+#include "GenericPlatformChunkInstall.h"
+
+DEFINE_LOG_CATEGORY(LogChunkInstaller)
\ No newline at end of file
diff --git a/Engine/Source/Runtime/Core/Private/Misc/EngineVersion.cpp b/Engine/Source/Runtime/Core/Private/Misc/EngineVersion.cpp
--- a/Engine/Source/Runtime/Core/Private/Misc/EngineVersion.cpp
+++ b/Engine/Source/Runtime/Core/Private/Misc/EngineVersion.cpp
@@ -14,10 +14,10 @@ const int32 GEngineMinNetVersion		= 7038;
 const int32 GEngineNegotiationVersion	= 3077;
 
 // Global instance of the current engine version
-FEngineVersion FEngineVersion::CurrentVersion(ENGINE_MAJOR_VERSION, ENGINE_MINOR_VERSION, ENGINE_PATCH_VERSION, (BUILT_FROM_CHANGELIST | (ENGINE_IS_LICENSEE_VERSION << 31)), BRANCH_NAME);
+FEngineVersion FEngineVersion::CurrentVersion(ENGINE_MAJOR_VERSION, ENGINE_MINOR_VERSION, ENGINE_PATCH_VERSION, (ENGINE_CURRENT_CL_VERSION | (ENGINE_IS_LICENSEE_VERSION << 31)), BRANCH_NAME);
 
 // Version which this engine maintains strict API and package compatibility with. By default, we always maintain compatibility with the current major/minor version, unless we're built at a different changelist.
-FEngineVersion FEngineVersion::CompatibleWithVersion(ENGINE_MAJOR_VERSION, ENGINE_MINOR_VERSION, 0, (MODULE_API_VERSION | (ENGINE_IS_LICENSEE_VERSION << 31)), BRANCH_NAME);
+FEngineVersion FEngineVersion::CompatibleWithVersion(ENGINE_MAJOR_VERSION, ENGINE_MINOR_VERSION, 0, (ENGINE_COMPATIBLE_CL_VERSION | (ENGINE_IS_LICENSEE_VERSION << 31)), BRANCH_NAME);
 
 
 FEngineVersionBase::FEngineVersionBase()
diff --git a/Engine/Source/Runtime/Core/Public/GenericPlatform/GenericPlatformChunkInstall.cpp b/Engine/Source/Runtime/Core/Public/GenericPlatform/GenericPlatformChunkInstall.cpp
deleted file mode 100644
diff --git a/Engine/Source/Runtime/Core/Public/GenericPlatform/GenericPlatformString.h b/Engine/Source/Runtime/Core/Public/GenericPlatform/GenericPlatformString.h
--- a/Engine/Source/Runtime/Core/Public/GenericPlatform/GenericPlatformString.h
+++ b/Engine/Source/Runtime/Core/Public/GenericPlatform/GenericPlatformString.h
@@ -108,7 +108,7 @@ struct FGenericPlatformString
 	>::Type Convert(DestEncoding* Dest, int32 DestSize, const SourceEncoding* Src, int32 SrcSize, DestEncoding BogusChar = (DestEncoding)'?')
 	{
 		if (DestSize < SrcSize)
-			return NULL;
+			return nullptr;
 
 		return (DestEncoding*)Memcpy(Dest, Src, SrcSize * sizeof(SourceEncoding)) + SrcSize;
 	}
diff --git a/Engine/Source/Runtime/Core/Public/Modules/ModuleVersion.h b/Engine/Source/Runtime/Core/Public/Modules/ModuleVersion.h
--- a/Engine/Source/Runtime/Core/Public/Modules/ModuleVersion.h
+++ b/Engine/Source/Runtime/Core/Public/Modules/ModuleVersion.h
@@ -4,19 +4,5 @@
 
 #include "../../../Launch/Resources/Version.h"
 
-// This number identifies a particular API revision, and is used to determine module compatibility. Hotfixes should retain the API version of the original release.
-// This define is parsed by the build tools, and should be a number or BUILT_FROM_CHANGELIST.
-#if BUILT_FROM_CHANGELIST > 0
-	#if ENGINE_IS_LICENSEE_VERSION
-		#define MODULE_API_VERSION BUILT_FROM_CHANGELIST
-	#else
-		#define MODULE_API_VERSION BUILT_FROM_CHANGELIST /* Or hotfix compatibility changelist */
-	#endif
-#else
-	#define MODULE_API_VERSION 0
-#endif
+#define MODULE_API_VERSION ENGINE_COMPATIBLE_CL_VERSION
 
-// Check that the API version has been set manually for a hotfix release
-#if ENGINE_IS_LICENSEE_VERSION == 0 && ENGINE_PATCH_VERSION > 0 && BUILT_FROM_CHANGELIST > 0 && MODULE_API_VERSION == BUILT_FROM_CHANGELIST
-	#error MODULE_API_VERSION must be manually defined for hotfix builds
-#endif
diff --git a/Engine/Source/Runtime/CoreUObject/Classes/Object.h b/Engine/Source/Runtime/CoreUObject/Classes/Object.h
--- a/Engine/Source/Runtime/CoreUObject/Classes/Object.h
+++ b/Engine/Source/Runtime/CoreUObject/Classes/Object.h
@@ -217,7 +217,7 @@ namespace EMouseCursor
 
 /** A set of numerical unit types supported by the engine. Mirrored from UnitConversion.h */
 UENUM(BlueprintType)
-enum class EUnit : uint8
+enum class EUnit
 {
 	/** Scalar distance/length unit. */
 
diff --git a/Engine/Source/Runtime/CoreUObject/Private/CoreUObjectPrivate.h b/Engine/Source/Runtime/CoreUObject/Private/CoreUObjectPrivate.h
--- a/Engine/Source/Runtime/CoreUObject/Private/CoreUObjectPrivate.h
+++ b/Engine/Source/Runtime/CoreUObject/Private/CoreUObjectPrivate.h
@@ -13,10 +13,21 @@
 #include "ScriptInterface.h"			// Script interface definitions.
 #include "Script.h"						// Script constants and forward declarations.
 #include "ObjectBase.h"					// Object base class.
+#include "UObjectAllocator.h"
+#include "UObjectGlobals.h"
+#include "UObjectMarks.h"
+#include "UObjectBase.h"
+#include "UObjectBaseUtility.h"
+#include "UObjectArray.h"
+#include "UObjectHash.h"
+#include "WeakObjectPtr.h"
+#include "UObject.h"
+#include "UObjectIterator.h"
 #include "CoreNet.h"					// Core networking.
 #include "ArchiveUObjectBase.h"			// UObject-related Archive classes.
 #include "GarbageCollection.h"			// Realtime garbage collection helpers
 #include "Class.h"						// Class definition.
+#include "StructOnScope.h"
 #include "Casts.h"                      // Cast templates
 #include "LazyObjectPtr.h"				// Object pointer types
 #include "AssetPtr.h"					// Object pointer types
diff --git a/Engine/Source/Runtime/CoreUObject/Private/UObject/UObjectHash.cpp b/Engine/Source/Runtime/CoreUObject/Private/UObject/UObjectHash.cpp
--- a/Engine/Source/Runtime/CoreUObject/Private/UObject/UObjectHash.cpp
+++ b/Engine/Source/Runtime/CoreUObject/Private/UObject/UObjectHash.cpp
@@ -5,6 +5,7 @@
 =============================================================================*/
 
 #include "CoreUObjectPrivate.h"
+#include "UObjectHash.h"
 
 DEFINE_LOG_CATEGORY_STATIC(LogUObjectHash, Log, All);
 
diff --git a/Engine/Source/Runtime/CoreUObject/Public/CoreUObject.h b/Engine/Source/Runtime/CoreUObject/Public/CoreUObject.h
--- a/Engine/Source/Runtime/CoreUObject/Public/CoreUObject.h
+++ b/Engine/Source/Runtime/CoreUObject/Public/CoreUObject.h
@@ -1,8 +1,6 @@
 // Copyright 1998-2016 Epic Games, Inc. All Rights Reserved.
 
-#ifndef __CoreUObject_h__
-#define __CoreUObject_h__
-
+#pragma once
 
 #include "Core.h"
 #include "ModuleManager.h"
@@ -10,11 +8,23 @@
 #include "Script.h"						// Script constants and forward declarations.
 #include "ScriptMacros.h"				// Script macro definitions
 #include "ObjectBase.h"					// Object base class.
+#include "UObjectAllocator.h"
+#include "UObjectGlobals.h"
+#include "UObjectMarks.h"
+#include "UObjectBase.h"
+#include "UObjectBaseUtility.h"
+#include "UObjectArray.h"
+#include "UObjectHash.h"
+#include "WeakObjectPtr.h"
+#include "UObject.h"
+#include "UObjectIterator.h"
 #include "CoreNet.h"					// Core networking.
 #include "ArchiveUObjectBase.h"			// UObject-related Archive classes.
 #include "GarbageCollection.h"			// Realtime garbage collection helpers
 #include "TextBuffer.h"					// UObjects for text buffers
 #include "Class.h"						// Class definition.
+#include "Templates/SubclassOf.h"
+#include "StructOnScope.h"
 #include "Casts.h"                      // Cast templates
 #include "LazyObjectPtr.h"				// Object pointer types
 #include "AssetPtr.h"					// Object pointer types
@@ -40,5 +50,3 @@
 #include "ScriptStackTracker.h"
 #include "WorldCompositionUtility.h"
 #include "StringClassReference.h"
-
-#endif	// __CoreUObject_h__
diff --git a/Engine/Source/Runtime/CoreUObject/Public/Misc/PackageName.h b/Engine/Source/Runtime/CoreUObject/Public/Misc/PackageName.h
--- a/Engine/Source/Runtime/CoreUObject/Public/Misc/PackageName.h
+++ b/Engine/Source/Runtime/CoreUObject/Public/Misc/PackageName.h
@@ -4,8 +4,7 @@
 	PackageName.h: Unreal package name utility functions.
 =============================================================================*/
 
-#ifndef __PACKAGENAME_H__
-#define __PACKAGENAME_H__
+#pragma once
 
 class COREUOBJECT_API FPackageName
 {
@@ -410,5 +409,3 @@ private:
 	static FOnContentPathDismountedEvent OnContentPathDismountedEvent;
 };
 
-#endif	// __PACKAGENAME_H__
-
diff --git a/Engine/Source/Runtime/CoreUObject/Public/Misc/StartupPackages.h b/Engine/Source/Runtime/CoreUObject/Public/Misc/StartupPackages.h
--- a/Engine/Source/Runtime/CoreUObject/Public/Misc/StartupPackages.h
+++ b/Engine/Source/Runtime/CoreUObject/Public/Misc/StartupPackages.h
@@ -4,9 +4,7 @@
 	StartupPackages.h: Startup Package Functions
 =============================================================================*/
 
-#ifndef __StartupPackages_H__
-#define __StartupPackages_H__
-
+#pragma once
 
 struct FStartupPackages
 {
@@ -38,7 +36,3 @@ struct FStartupPackages
 	 */
 	COREUOBJECT_API static bool LoadAll();
 };
-
-
-#endif //__StartupPackages_H__
-
diff --git a/Engine/Source/Runtime/CoreUObject/Public/Misc/WorldCompositionUtility.h b/Engine/Source/Runtime/CoreUObject/Public/Misc/WorldCompositionUtility.h
--- a/Engine/Source/Runtime/CoreUObject/Public/Misc/WorldCompositionUtility.h
+++ b/Engine/Source/Runtime/CoreUObject/Public/Misc/WorldCompositionUtility.h
@@ -6,9 +6,6 @@
 
 #pragma once
 
-#ifndef __WorldCompositionUtility_h__
-#define __WorldCompositionUtility_h__
-
 #define WORLDTILE_LOD_PACKAGE_SUFFIX	TEXT("_LOD")
 #define WORLDTILE_LOD_MAX_INDEX			4
 
@@ -141,5 +138,3 @@ public:
 	int32				ZOrder;
 };
 
-
-#endif	// __WorldCompositionUtility_h__
\ No newline at end of file
diff --git a/Engine/Source/Runtime/CoreUObject/Public/Templates/Casts.h b/Engine/Source/Runtime/CoreUObject/Public/Templates/Casts.h
--- a/Engine/Source/Runtime/CoreUObject/Public/Templates/Casts.h
+++ b/Engine/Source/Runtime/CoreUObject/Public/Templates/Casts.h
@@ -3,6 +3,7 @@
 #pragma once
 
 #include "Templates/AndOr.h"
+#include "SubclassOf.h"
 
 COREUOBJECT_API void CastLogError(const TCHAR* FromType, const TCHAR* ToType);
 
diff --git a/Engine/Source/Runtime/CoreUObject/Public/Templates/SubclassOf.h b/Engine/Source/Runtime/CoreUObject/Public/Templates/SubclassOf.h
new file mode 100644
--- /dev/null
+++ b/Engine/Source/Runtime/CoreUObject/Public/Templates/SubclassOf.h
@@ -0,0 +1,118 @@
+// Copyright 1998-2016 Epic Games, Inc. All Rights Reserved.
+
+#pragma once
+
+#include "Class.h"
+
+/**
+* Template to allow UClass's to be passed around with type safety
+**/
+template<class TClass>
+class TSubclassOf
+{
+public:
+	/**
+	* Default Constructor, defaults to NULL
+	**/
+	FORCEINLINE TSubclassOf() :
+		Class(NULL)
+	{
+	}
+	/**
+	* Constructor that takes a UClass and does a runtime check to make sure this is a compatible class
+	* @param From UClass to assign to this TSubclassOf
+	**/
+	FORCEINLINE TSubclassOf(UClass* From) :
+		Class(From)
+	{
+	}
+	/**
+	* Copy Constructor
+	* @param From TSubclassOf to copy
+	**/
+	template<class TClassA>
+	FORCEINLINE TSubclassOf(const TSubclassOf<TClassA>& From) :
+		Class(*From)
+	{
+		TClass* Test = (TClassA*)0; // verify that TClassA derives from TClass
+	}
+
+	/**
+	* Assignment operator
+	* @param From TSubclassOf to copy
+	**/
+	template<class TClassA>
+	FORCEINLINE TSubclassOf& operator=(const TSubclassOf<TClassA>& From)
+	{
+		Class = *From;
+		TClass* Test = (TClassA*)0; // verify that TClassA derives from TClass
+		return *this;
+	}
+	/**
+	* Assignment operator from UClass, runtime check for compatibility
+	* @param From UClass to assign to this TSubclassOf
+	**/
+	FORCEINLINE TSubclassOf& operator=(UClass* From)
+	{
+		Class = From;
+		return *this;
+	}
+	/**
+	* Dereference back into a UClass
+	* @return	the embedded UClass
+	* The body for this method is in Class.h
+	*/
+	FORCEINLINE UClass* operator*() const;
+
+	/**
+	* Dereference back into a UClass
+	* @return	the embedded UClass
+	*/
+	FORCEINLINE UClass* operator->() const
+	{
+		return **this;
+	}
+
+	/**
+	* Implicit conversion to UClass
+	* @return	the embedded UClass
+	*/
+	FORCEINLINE operator UClass* () const
+	{
+		return **this;
+	}
+
+	/**
+	* Get the CDO if we are referencing a valid class
+	* @return the CDO, or NULL
+	*/
+	FORCEINLINE TClass* GetDefaultObject() const;
+
+	friend FArchive& operator<<(FArchive& Ar, TSubclassOf& SubclassOf)
+	{
+		Ar << SubclassOf.Class;
+		return Ar;
+	}
+private:
+	UClass* Class;
+};
+
+/**
+* Dereference back into a UClass
+* @return	the embedded UClass
+*/
+template<class TClass>
+FORCEINLINE UClass* TSubclassOf<TClass>::operator*() const
+{
+	if (!Class || !Class->IsChildOf(TClass::StaticClass()))
+	{
+		return NULL;
+	}
+	return Class;
+}
+
+template<class TClass>
+FORCEINLINE TClass* TSubclassOf<TClass>::GetDefaultObject() const
+{
+	return Class ? Class->GetDefaultObject<TClass>() : NULL;
+}
diff --git a/Engine/Source/Runtime/CoreUObject/Public/UObject/Class.h b/Engine/Source/Runtime/CoreUObject/Public/UObject/Class.h
--- a/Engine/Source/Runtime/CoreUObject/Public/UObject/Class.h
+++ b/Engine/Source/Runtime/CoreUObject/Public/UObject/Class.h
@@ -7,6 +7,8 @@
 #pragma once
 
 #include "ObjectBase.h"
+#include "UObject.h"
+#include "GarbageCollection.h"
 
 /*-----------------------------------------------------------------------------
 	Mirrors of mirror structures in Object.h. These are used by generated code 
@@ -18,6 +20,7 @@ COREUOBJECT_API DECLARE_LOG_CATEGORY_EXTERN(LogClass, Log, All);
 COREUOBJECT_API DECLARE_LOG_CATEGORY_EXTERN(LogScriptSerialization, Log, All);
 
 struct FPropertyTag;
+struct FNetDeltaSerializeInfo;
 
 /*-----------------------------------------------------------------------------
 	FRepRecord.
@@ -1162,119 +1165,6 @@ public:
 #endif // WITH_EDITOR
 };
 
-class FStructOnScope
-{
-protected:
-	TWeakObjectPtr<const UStruct> ScriptStruct;
-	uint8* SampleStructMemory;
-	TWeakObjectPtr<UPackage> Package;
-
-	FStructOnScope()
-		: SampleStructMemory(nullptr)
-		, OwnsMemory(false)
-	{
-	}
-
-	void Initialize()
-	{
-		if (ScriptStruct.IsValid())
-		{
-			SampleStructMemory = (uint8*)FMemory::Malloc(ScriptStruct->GetStructureSize());
-			ScriptStruct.Get()->InitializeStruct(SampleStructMemory);
-			OwnsMemory = true;
-		}
-	}
-
-public:
-
-	FStructOnScope(const UStruct* InScriptStruct)
-		: ScriptStruct(InScriptStruct)
-		, SampleStructMemory(nullptr)
-		, OwnsMemory(false)
-	{
-		Initialize();
-	}
-
-	FStructOnScope(const UStruct* InScriptStruct, uint8* InData)
-		: ScriptStruct(InScriptStruct)
-		, SampleStructMemory(InData)
-		, OwnsMemory(false)
-	{
-	}
-
-	virtual uint8* GetStructMemory()
-	{
-		return SampleStructMemory;
-	}
-
-	virtual const uint8* GetStructMemory() const
-	{
-		return SampleStructMemory;
-	}
-
-	virtual const UStruct* GetStruct() const
-	{
-		return ScriptStruct.Get();
-	}
-
-	virtual UPackage* GetPackage() const
-	{
-		return Package.Get();
-	}
-
-	virtual void SetPackage(UPackage* InPackage)
-	{
-		Package = InPackage;
-	}
-
-	virtual bool IsValid() const
-	{
-		return ScriptStruct.IsValid() && SampleStructMemory;
-	}
-
-	virtual void Destroy()
-	{
-		if (!OwnsMemory)
-		{
-			return;
-		}
-
-		if (ScriptStruct.IsValid() && SampleStructMemory)
-		{
-			ScriptStruct.Get()->DestroyStruct(SampleStructMemory);
-			ScriptStruct = NULL;
-		}
-
-		if (SampleStructMemory)
-		{
-			FMemory::Free(SampleStructMemory);
-			SampleStructMemory = NULL;
-		}
-	}
-
-	virtual ~FStructOnScope()
-	{
-		Destroy();
-	}
-
-	/** Re-initializes the scope with a specified UStruct */
-	void Initialize(TWeakObjectPtr<const UStruct> InScriptStruct)
-	{
-		ScriptStruct = InScriptStruct;
-		Initialize();
-	}
-
-private:
-
-	FStructOnScope(const FStructOnScope&);
-	FStructOnScope& operator=(const FStructOnScope&);
-
-private:
-
-	/** Whether the struct memory is owned by this instance. */
-	bool OwnsMemory;
-};
-
 /*-----------------------------------------------------------------------------
 	UFunction.
 -----------------------------------------------------------------------------*/
@@ -1361,12 +1251,7 @@ public:
 	virtual void Link(FArchive& Ar, bool bRelinkExistingProperties) override;
 
 	// UFunction interface.
-	UFunction* GetSuperFunction() const
-	{
-		UStruct* Result = GetSuperStruct();
-		checkSlow(!Result || Result->IsA<UFunction>());
-		return (UFunction*)Result;
-	}
+	UFunction* GetSuperFunction() const;
 
 	UProperty* GetReturnProperty() const;
 
@@ -2816,28 +2701,15 @@ private:
 	TMap<class UObject*,class UObject*>			SourceToDestinationMap;
 };
 
+// UFunction interface.
 
-// ObjectBase.h
-
-/**
- * Dereference back into a UClass
- * @return	the embedded UClass
- */
-template<class TClass>
-FORCEINLINE UClass* TSubclassOf<TClass>::operator*() const
-{
-	if (!Class || !Class->IsChildOf(TClass::StaticClass()))
+inline UFunction* UFunction::GetSuperFunction() const
 {
-		return NULL;
-	}
-	return Class;
+	UStruct* Result = GetSuperStruct();
+	checkSlow(!Result || Result->IsA<UFunction>());
+	return (UFunction*)Result;
 }
 
-template<class TClass>
-FORCEINLINE TClass* TSubclassOf<TClass>::GetDefaultObject() const
-{
-	return Class ? Class->GetDefaultObject<TClass>() : NULL;
-}
 
 // UObject.h
 
@@ -3004,4 +2876,5 @@ template<> struct TBaseStructure<FFloatInterval>
 template<> struct TBaseStructure<FInt32Interval>
 {
 	COREUOBJECT_API static UScriptStruct* Get();
-};
\ No newline at end of file
+};
+
diff --git a/Engine/Source/Runtime/CoreUObject/Public/UObject/ClassTree.h b/Engine/Source/Runtime/CoreUObject/Public/UObject/ClassTree.h
--- a/Engine/Source/Runtime/CoreUObject/Public/UObject/ClassTree.h
+++ b/Engine/Source/Runtime/CoreUObject/Public/UObject/ClassTree.h
@@ -4,8 +4,7 @@
 	ClassTree.h: Class hierarchy management classes.
 =============================================================================*/
 
-#ifndef __CLASSTREE_H__
-#define __CLASSTREE_H__
+#pragma once
 
 enum EClassFlagMatchType
 {
@@ -601,5 +600,3 @@ public:
 		return Children[Index];
 	}
 };
-
-#endif	// __CLASSTREE_H__
diff --git a/Engine/Source/Runtime/CoreUObject/Public/UObject/CoreNet.h b/Engine/Source/Runtime/CoreUObject/Public/UObject/CoreNet.h
--- a/Engine/Source/Runtime/CoreUObject/Public/UObject/CoreNet.h
+++ b/Engine/Source/Runtime/CoreUObject/Public/UObject/CoreNet.h
@@ -3,6 +3,8 @@
 #pragma once
 
 #include "ObjectBase.h"
+#include "WeakObjectPtr.h"
+#include "UObject.h"
 
 
 DECLARE_DELEGATE_RetVal_OneParam( bool, FNetObjectIsDynamic, const UObject*);
@@ -275,6 +277,7 @@ public:
 
 template <> struct TIsZeroConstructType<FLifetimeProperty> { enum { Value = true }; };
 
+GENERATE_MEMBER_FUNCTION_CHECK(GetLifetimeReplicatedProps, void, const, TArray<FLifetimeProperty>&)
 
 /**
  * FNetBitWriter
diff --git a/Engine/Source/Runtime/CoreUObject/Public/UObject/GarbageCollection.h b/Engine/Source/Runtime/CoreUObject/Public/UObject/GarbageCollection.h
--- a/Engine/Source/Runtime/CoreUObject/Public/UObject/GarbageCollection.h
+++ b/Engine/Source/Runtime/CoreUObject/Public/UObject/GarbageCollection.h
@@ -1,10 +1,13 @@
 // Copyright 1998-2016 Epic Games, Inc. All Rights Reserved.
 
-#pragma once
 /*=============================================================================
 	GarbageCollection.h: Unreal realtime garbage collection helpers
 =============================================================================*/
 
+#pragma once
+
+#include "ArchiveUObjectBase.h"
+
 /** Context sensitive keep flags for garbage collection */
 #define GARBAGE_COLLECTION_KEEPFLAGS	(GIsEditor ? RF_Standalone : RF_NoFlags)
 
diff --git a/Engine/Source/Runtime/CoreUObject/Public/UObject/ObjectBase.h b/Engine/Source/Runtime/CoreUObject/Public/UObject/ObjectBase.h
--- a/Engine/Source/Runtime/CoreUObject/Public/UObject/ObjectBase.h
+++ b/Engine/Source/Runtime/CoreUObject/Public/UObject/ObjectBase.h
@@ -1636,119 +1636,3 @@ public:
 };
 
 
-#include "UObjectAllocator.h"
-#include "UObjectGlobals.h"
-#include "UObjectMarks.h"
-#include "UObjectBase.h"
-#include "UObjectBaseUtility.h"
-#include "UObjectArray.h"
-#include "UObjectHash.h"
-#include "WeakObjectPtr.h"
-#include "UObject.h"
-#include "UObjectIterator.h"
-
-/**
- * Test validity of object 
- *
- * @param	Test			The object to test
- * @return	Return true if the object is usable: non-null and not pending kill
- */
-FORCEINLINE bool IsValid(const UObject *Test)
-{
-	return Test && !Test->IsPendingKill();
-}
-
-/**
- * Template to allow UClass's to be passed around with type safety
-**/
-template<class TClass>
-class TSubclassOf
-{
-public:
-	/**
-	 * Default Constructor, defaults to NULL
-	**/
-	FORCEINLINE TSubclassOf() : 
-		Class(NULL)
-	{
-	}
-	/**
-	 * Constructor that takes a UClass and does a runtime check to make sure this is a compatible class
-	 * @param From UClass to assign to this TSubclassOf
-	**/
-	FORCEINLINE TSubclassOf(UClass* From) : 
-		Class(From)
-	{
-	}
-	/**
-	 * Copy Constructor 
-	 * @param From TSubclassOf to copy
-	**/
-	template<class TClassA>
-	FORCEINLINE TSubclassOf(const TSubclassOf<TClassA>& From) : 
-		Class(*From)
-	{
-		TClass* Test = (TClassA*)0; // verify that TClassA derives from TClass
-	}
-
-	/**
-	 * Assignment operator
-	 * @param From TSubclassOf to copy
-	**/
-	template<class TClassA>
-	FORCEINLINE TSubclassOf& operator=(const TSubclassOf<TClassA>& From)
-	{
-		Class = *From;
-		TClass* Test = (TClassA*)0; // verify that TClassA derives from TClass
-		return *this;
-	}
-	/**
-	 * Assignment operator from UClass, runtime check for compatibility
-	 * @param From UClass to assign to this TSubclassOf
-	**/
-	FORCEINLINE TSubclassOf& operator=(UClass* From)
-	{
-		Class = From;
-		return *this;
-	}
-	/**
- 	 * Dereference back into a UClass
-	 * @return	the embedded UClass
-	 * The body for this method is in Class.h
-	 */
-	FORCEINLINE UClass* operator*() const;
-
-	/**
- 	 * Dereference back into a UClass
-	 * @return	the embedded UClass
-	 */
-	FORCEINLINE UClass* operator->() const
-	{
-		return **this;
-	}
-
-	/**
-	 * Implicit conversion to UClass
-	 * @return	the embedded UClass
-	 */
-	FORCEINLINE operator UClass* () const
-	{
-		return **this;
-	}
-
-	/**
-	 * Get the CDO if we are referencing a valid class
-	 * @return the CDO, or NULL
-	 */
-	FORCEINLINE TClass* GetDefaultObject() const;
-	
-	friend FArchive& operator<<( FArchive& Ar, TSubclassOf& SubclassOf )
-	{
-		Ar << SubclassOf.Class;
-		return Ar;
-	}
-private:
-	UClass* Class;
-};
-
-GENERATE_MEMBER_FUNCTION_CHECK(GetLifetimeReplicatedProps, void, const, TArray<FLifetimeProperty>&)
diff --git a/Engine/Source/Runtime/CoreUObject/Public/UObject/PropertyTag.h b/Engine/Source/Runtime/CoreUObject/Public/UObject/PropertyTag.h
--- a/Engine/Source/Runtime/CoreUObject/Public/UObject/PropertyTag.h
+++ b/Engine/Source/Runtime/CoreUObject/Public/UObject/PropertyTag.h
@@ -4,8 +4,7 @@
 	FPropertyTag.
 -----------------------------------------------------------------------------*/
 
-#ifndef __UNPROPERTYTAG_H__
-#define __UNPROPERTYTAG_H__
+#pragma once
 
 #include "DebugSerializationFlags.h"
 
@@ -43,5 +42,3 @@ struct FPropertyTag
 	// Property serializer.
 	void SerializeTaggedProperty( FArchive& Ar, UProperty* Property, uint8* Value, uint8* Defaults );
 };
-
-#endif	// __UNPROPERTYTAG_H__
diff --git a/Engine/Source/Runtime/CoreUObject/Public/UObject/ScriptInterface.h b/Engine/Source/Runtime/CoreUObject/Public/UObject/ScriptInterface.h
--- a/Engine/Source/Runtime/CoreUObject/Public/UObject/ScriptInterface.h
+++ b/Engine/Source/Runtime/CoreUObject/Public/UObject/ScriptInterface.h
@@ -4,8 +4,7 @@
 	ScriptInterface.h: Script interface definitions.
 =============================================================================*/
 
-#ifndef __SCRIPTINTERFACE_H__
-#define __SCRIPTINTERFACE_H__
+#pragma once
 
 /**
  * FScriptInterface
@@ -220,5 +219,3 @@ public:
 		return Interface.Serialize(Ar, InterfaceType::UClassType::StaticClass());
 	}
 };
-
-#endif
diff --git a/Engine/Source/Runtime/CoreUObject/Public/UObject/StructOnScope.h b/Engine/Source/Runtime/CoreUObject/Public/UObject/StructOnScope.h
new file mode 100644
--- /dev/null
+++ b/Engine/Source/Runtime/CoreUObject/Public/UObject/StructOnScope.h
@@ -0,0 +1,118 @@
+// Copyright 1998-2016 Epic Games, Inc. All Rights Reserved.
+
+#pragma once
+
+#include "CoreObject.h"
+
+class FStructOnScope
+{
+protected:
+	TWeakObjectPtr<const UStruct> ScriptStruct;
+	uint8* SampleStructMemory;
+	TWeakObjectPtr<UPackage> Package;
+
+	FStructOnScope()
+		: SampleStructMemory(nullptr)
+		, OwnsMemory(false)
+	{
+	}
+
+	void Initialize()
+	{
+		if (ScriptStruct.IsValid())
+		{
+			SampleStructMemory = (uint8*)FMemory::Malloc(ScriptStruct->GetStructureSize());
+			ScriptStruct.Get()->InitializeStruct(SampleStructMemory);
+			OwnsMemory = true;
+		}
+	}
+
+public:
+
+	FStructOnScope(const UStruct* InScriptStruct)
+		: ScriptStruct(InScriptStruct)
+		, SampleStructMemory(nullptr)
+		, OwnsMemory(false)
+	{
+		Initialize();
+	}
+
+	FStructOnScope(const UStruct* InScriptStruct, uint8* InData)
+		: ScriptStruct(InScriptStruct)
+		, SampleStructMemory(InData)
+		, OwnsMemory(false)
+	{
+	}
+
+	virtual uint8* GetStructMemory()
+	{
+		return SampleStructMemory;
+	}
+
+	virtual const uint8* GetStructMemory() const
+	{
+		return SampleStructMemory;
+	}
+
+	virtual const UStruct* GetStruct() const
+	{
+		return ScriptStruct.Get();
+	}
+
+	virtual UPackage* GetPackage() const
+	{
+		return Package.Get();
+	}
+
+	virtual void SetPackage(UPackage* InPackage)
+	{
+		Package = InPackage;
+	}
+
+	virtual bool IsValid() const
+	{
+		return ScriptStruct.IsValid() && SampleStructMemory;
+	}
+
+	virtual void Destroy()
+	{
+		if (!OwnsMemory)
+		{
+			return;
+		}
+
+		if (ScriptStruct.IsValid() && SampleStructMemory)
+		{
+			ScriptStruct.Get()->DestroyStruct(SampleStructMemory);
+			ScriptStruct = NULL;
+		}
+
+		if (SampleStructMemory)
+		{
+			FMemory::Free(SampleStructMemory);
+			SampleStructMemory = NULL;
+		}
+	}
+
+	virtual ~FStructOnScope()
+	{
+		Destroy();
+	}
+
+	/** Re-initializes the scope with a specified UStruct */
+	void Initialize(TWeakObjectPtr<const UStruct> InScriptStruct)
+	{
+		ScriptStruct = InScriptStruct;
+		Initialize();
+	}
+
+private:
+
+	FStructOnScope(const FStructOnScope&);
+	FStructOnScope& operator=(const FStructOnScope&);
+
+private:
+
+	/** Whether the struct memory is owned by this instance. */
+	bool OwnsMemory;
+};
diff --git a/Engine/Source/Runtime/CoreUObject/Public/UObject/UObject.h b/Engine/Source/Runtime/CoreUObject/Public/UObject/UObject.h
--- a/Engine/Source/Runtime/CoreUObject/Public/UObject/UObject.h
+++ b/Engine/Source/Runtime/CoreUObject/Public/UObject/UObject.h
@@ -2,6 +2,7 @@
 
 #pragma once
 
+#include "UObjectBaseUtility.h"
 
 DECLARE_LOG_CATEGORY_EXTERN(LogObj, Log, All);
 
@@ -1182,3 +1183,14 @@ private:
 	*/
 	UObject* CreateEditorOnlyDefaultSubobjectImpl(FName SubobjectName, UClass* ReturnType, bool bTransient = false);
 };
+
+/**
+* Test validity of object
+*
+* @param	Test			The object to test
+* @return	Return true if the object is usable: non-null and not pending kill
+*/
+FORCEINLINE bool IsValid(const UObject *Test)
+{
+	return Test && !Test->IsPendingKill();
+}
diff --git a/Engine/Source/Runtime/CoreUObject/Public/UObject/UObjectAllocator.h b/Engine/Source/Runtime/CoreUObject/Public/UObject/UObjectAllocator.h
--- a/Engine/Source/Runtime/CoreUObject/Public/UObject/UObjectAllocator.h
+++ b/Engine/Source/Runtime/CoreUObject/Public/UObject/UObjectAllocator.h
@@ -4,8 +4,8 @@
 	UnObjAllocator.h: Unreal object allocation
 =============================================================================*/
 
-#ifndef __UNOBJALLOCATOR_H__
-#define __UNOBJALLOCATOR_H__
+#pragma once
+
 #include "HAL/Platform.h"
 
 class UObjectBase;
@@ -80,6 +80,3 @@ private:
 /** Global UObjectBase allocator							*/
 extern COREUOBJECT_API FUObjectAllocator GUObjectAllocator;
 
-
-#endif	// __UNOBJALLOCATOR_H__
-
diff --git a/Engine/Source/Runtime/CoreUObject/Public/UObject/UObjectAnnotation.h b/Engine/Source/Runtime/CoreUObject/Public/UObject/UObjectAnnotation.h
--- a/Engine/Source/Runtime/CoreUObject/Public/UObject/UObjectAnnotation.h
+++ b/Engine/Source/Runtime/CoreUObject/Public/UObject/UObjectAnnotation.h
@@ -4,8 +4,7 @@
 	UObjectAnnotation.h: Unreal object annotation template
 =============================================================================*/
 
-#ifndef __UNOBJANNOTATION_H__
-#define __UNOBJANNOTATION_H__
+#pragma once
 
 /**
 * FUObjectAnnotationSparse is a helper class that is used to store sparse, slow, temporary, editor only, external 
@@ -780,6 +779,3 @@ private:
 // Definition is in UObjectGlobals.cpp
 extern COREUOBJECT_API FUObjectAnnotationSparseBool GSelectedAnnotation;
 
-
-#endif	// __UNOBJANNOTATION_H__
-
diff --git a/Engine/Source/Runtime/CoreUObject/Public/UObject/UObjectArray.h b/Engine/Source/Runtime/CoreUObject/Public/UObject/UObjectArray.h
--- a/Engine/Source/Runtime/CoreUObject/Public/UObject/UObjectArray.h
+++ b/Engine/Source/Runtime/CoreUObject/Public/UObject/UObjectArray.h
@@ -4,8 +4,9 @@
 	UObjectArray.h: Unreal object array
 =============================================================================*/
 
-#ifndef __UOBJECTARRAY_H__
-#define __UOBJECTARRAY_H__
+#pragma once
+
+#include "UObjectBase.h"
 
 /**
 * Controls whether the number of available elements is being tracked in the ObjObjects array.
@@ -770,5 +771,3 @@ struct FIndexToObject
 		return ObjectItem ? ObjectItem->Object : nullptr;
 	}
 };
-
-#endif	// __UOBJECTARRAY_H__
diff --git a/Engine/Source/Runtime/CoreUObject/Public/UObject/UObjectBase.h b/Engine/Source/Runtime/CoreUObject/Public/UObject/UObjectBase.h
--- a/Engine/Source/Runtime/CoreUObject/Public/UObject/UObjectBase.h
+++ b/Engine/Source/Runtime/CoreUObject/Public/UObject/UObjectBase.h
@@ -4,8 +4,9 @@
 	UObjectBase.h: Unreal UObject base class
 =============================================================================*/
 
-#ifndef __UOBJECTBASE_H__
-#define __UOBJECTBASE_H__
+#pragma once
+
+#include "UObjectGlobals.h"
 
 DECLARE_DWORD_COUNTER_STAT_EXTERN(TEXT("STAT_UObjectsStatGroupTester"), STAT_UObjectsStatGroupTester, STATGROUP_UObjects, COREUOBJECT_API);
 
@@ -388,5 +389,3 @@ void UObjectBaseInit();
  */
 void UObjectBaseShutdown();
 
-#endif	// __UOBJECTBASE_H__
-
diff --git a/Engine/Source/Runtime/CoreUObject/Public/UObject/UObjectBaseUtility.h b/Engine/Source/Runtime/CoreUObject/Public/UObject/UObjectBaseUtility.h
--- a/Engine/Source/Runtime/CoreUObject/Public/UObject/UObjectBaseUtility.h
+++ b/Engine/Source/Runtime/CoreUObject/Public/UObject/UObjectBaseUtility.h
@@ -7,6 +7,7 @@
 #pragma once
 
 #include "UObjectArray.h"
+#include "UObjectMarks.h"
 
 #if _MSC_VER == 1900
 	#ifdef PRAGMA_DISABLE_SHADOW_VARIABLE_WARNINGS
diff --git a/Engine/Source/Runtime/CoreUObject/Public/UObject/UObjectGlobals.h b/Engine/Source/Runtime/CoreUObject/Public/UObject/UObjectGlobals.h
--- a/Engine/Source/Runtime/CoreUObject/Public/UObject/UObjectGlobals.h
+++ b/Engine/Source/Runtime/CoreUObject/Public/UObject/UObjectGlobals.h
@@ -4,8 +4,7 @@
 	UnObjGlobals.h: Unreal object system globals.
 =============================================================================*/
 
-#ifndef __UNOBJGLOBALS_H__
-#define __UNOBJGLOBALS_H__
+#pragma once
 
 #include "Script.h"
 
@@ -1956,5 +1955,3 @@ COREUOBJECT_API TMap<FName, FDynamicClassStaticData>& GetDynamicClassMap();
  */
 COREUOBJECT_API bool IsEditorOnlyObject(const UObject* InObject);
 #endif //WITH_EDITOR
-
-#endif	// __UNOBJGLOBALS_H__
diff --git a/Engine/Source/Runtime/CoreUObject/Public/UObject/UObjectHash.h b/Engine/Source/Runtime/CoreUObject/Public/UObject/UObjectHash.h
--- a/Engine/Source/Runtime/CoreUObject/Public/UObject/UObjectHash.h
+++ b/Engine/Source/Runtime/CoreUObject/Public/UObject/UObjectHash.h
@@ -4,8 +4,7 @@
 	UObjectHash.h: Unreal object name hashes
 =============================================================================*/
 
-#ifndef __UOBJECTHASH_H__
-#define __UOBJECTHASH_H__
+#pragma once
 
 DECLARE_STATS_GROUP_VERBOSE(TEXT("UObject Hash"), STATGROUP_UObjectHash, STATCAT_Advanced);
 
@@ -136,6 +135,3 @@ COREUOBJECT_API void LogHashOuterStatistics(FOutputDevice& Ar, const bool bShowH
  */
 void AllocateUObjectIndexForCurrentThread(class UObjectBase* Object);
 
-
-#endif	// __UOBJECTHASH_H__
-
diff --git a/Engine/Source/Runtime/CoreUObject/Public/UObject/UObjectIterator.h b/Engine/Source/Runtime/CoreUObject/Public/UObject/UObjectIterator.h
--- a/Engine/Source/Runtime/CoreUObject/Public/UObject/UObjectIterator.h
+++ b/Engine/Source/Runtime/CoreUObject/Public/UObject/UObjectIterator.h
@@ -4,9 +4,9 @@
 	UObjectIterator.h: High level iterators for uobject
 =============================================================================*/
 
-#ifndef __UOBJECTITERATOR_H__
-#define __UOBJECTITERATOR_H__
+#pragma once
 
+#include "Class.h"
 
 /**
  * Class for iterating through all objects, including class default objects, unreachable objects...all UObjects
@@ -327,5 +327,3 @@ struct TObjectRange<UObject>
 	TObjectIterator<UObject> Begin;
 };
 
-
-#endif	// __UOBJECTITERATOR_H__
diff --git a/Engine/Source/Runtime/CoreUObject/Public/UObject/UObjectMarks.h b/Engine/Source/Runtime/CoreUObject/Public/UObject/UObjectMarks.h
--- a/Engine/Source/Runtime/CoreUObject/Public/UObject/UObjectMarks.h
+++ b/Engine/Source/Runtime/CoreUObject/Public/UObject/UObjectMarks.h
@@ -4,8 +4,7 @@
 	UObjectMark.h: Unreal object marks
 =============================================================================*/
 
-#ifndef __UOBJECTMARKS_H__
-#define __UOBJECTMARKS_H__
+#pragma once
 
 /** 
  * Object marks are bits associated with UObject. Marks are:
@@ -120,5 +119,3 @@ COREUOBJECT_API void GetObjectsWithAllMarks(TArray<UObject *>& Results, EObjectM
  */
 COREUOBJECT_API void GetObjectsWithAnyMarks(TArray<UObject *>& Results, EObjectMark Marks);
 
-#endif	// __UOBJECTMARKS_H__
-
diff --git a/Engine/Source/Runtime/CoreUObject/Public/UObject/WeakObjectPtr.h b/Engine/Source/Runtime/CoreUObject/Public/UObject/WeakObjectPtr.h
--- a/Engine/Source/Runtime/CoreUObject/Public/UObject/WeakObjectPtr.h
+++ b/Engine/Source/Runtime/CoreUObject/Public/UObject/WeakObjectPtr.h
@@ -7,6 +7,7 @@
 #pragma once
 
 #include "AutoPointer.h"
+#include "UObjectArray.h"
 
 /***
  * 
diff --git a/Engine/Source/Runtime/Launch/Resources/Version.h b/Engine/Source/Runtime/Launch/Resources/Version.h
--- a/Engine/Source/Runtime/Launch/Resources/Version.h
+++ b/Engine/Source/Runtime/Launch/Resources/Version.h
@@ -2,66 +2,124 @@
 
 #pragma once
 
-#define VERSION_TEXT_2(x) L ## x
-#define VERSION_TEXT(x) VERSION_TEXT_2(x)
-#define VERSION_STRINGIFY_2(x) VERSION_TEXT(#x)
-#define VERSION_STRINGIFY(x) VERSION_STRINGIFY_2(x)
-
-#include "VersionLocked.h"
-
-#if !LOCKED_BUILD_VERSION
-
+//=============================================================================================================================================================
+// This file defines constants used for versioning packages, modules, and various subsystems in UE4. In general, it should not be necessary to include this
+// file and access these values directly - they are all wrapped behind the higher-level abstractions in FEngineVersion and the FApp class.
+//
+// The following concepts are used for versioning in UE4:
+//
+//  -  The *engine version* defines the explicit major/minor/patch version of the engine, plus the changelist and branch name that it was built from. The 
+//     changelist is assumed to be a monotonically increasing number in the current branch, and is used both as a unique identifier and to infer that one engine
+//     was later than another. Tagged property serialization in UE4 is tolerant to properties being added or removed, so we always want to prevent an older 
+//     build of the engine loading assets created with a newer build, discarding properties which have recently been added, and silently losing data
+//     when the asset is saved out. The changelist allows ordering versions in such cases. The engine version is encapsulated by the FEngineVersion class, of 
+//     which there are two commonly referenced instances:
+//
+//         FEngineVersion::Current()        normally uses ENGINE_CURRENT_CL_VERSION for the changelist component, and indicates the code the engine was built 
+//                                          from. This is typically only used for diagnostic and display purposes.
+//
+//         FEngineVersion::CompatibleWith() normally uses ENGINE_COMPATIBLE_CL_VERSION for the changelist component and '0' for the patch component, and 
+//                                          indicates the baseline version of the engine that this build maintains strict binary compatibility with. 
+//                                          By default, this compatibility extends to assets, executable modules, and any network data transmitted between 
+//                                          two builds, and is used when creating patches and hotfixes that can be used interchangably with another build.
+//                                          This should be used for versioning in the majority of cases in the engine.
+//
+//     Both the ENGINE_CURRENT_CL_VERSION and ENGINE_COMPATIBLE_CL_VERSION macros can be updated systemically by build systems using the UpdateLocalVersion 
+//     AutomationTool command (as well as the ENGINE_IS_LICENSEE_VERSION and BRANCH_NAME macros).
+//
+//  -  The *object version* (aka serialization version) is as a monotonically incrementing (but manually updated) integer, and is used to write one-way 
+//     upgrade code in custom UObject serialization functions. It is set by the enum in ObjectVersion.h, and is global to the whole engine. This version number 
+//     is saved as a raw integer value in package headers, so it cannot be safely reordered or merged between branches. It should ONLY be updated 
+//     by Epic, otherwise future engine merges may corrupt content.
+//
+//  -  The *licensee object version* is provided for licensees to create their own one-way upgrade paths akin to the regular object version. Epic will never
+//     add entries to this enumeration. It is defined by the enum in ObjectVersion.h
+//
+//  -  Any number of *custom object version* objects may be registered to create orthoganal incrementing version numbers similar to the object version and 
+//     licensee version enums (see FCustomVersion). Each one is registered with a GUID, ensuring uniqueness and allowing the FArchive to quickly store and 
+//     retrieve them without any context of what they represent. Custom versions may be created for individual projects, subsystems, or branches.
+//
+//  -  The *build version* is an opaque string specific to the product being built, and should be used for identifying the current application 
+//     (as opposed to distinct applications built with the same engine version). It is set by the BUILD_VERSION macro, which can be updated using the 
+//     UpdateLocalVersion AutomationTool command.
+//
+//  -  The *network version* and *replay version* are used for versioning the network and replay subsystems, and default to the compatible engine version.
+//
+//  -  The *engine association* in a .uproject file often takes the appearance of a version number for launcher-installed binary UE4 releases, but may be 
+//     other identifiers as well. See ProjectDescriptor.h for a description of how this technique works.
+//
+// Constants in this file are updated by AutomationTool and UnrealGameSync. Be careful when changing formatting for the submitted version of this file that 
+// these tools can still parse it.
+//=============================================================================================================================================================
+
+// These numbers define the banner UE4 version, and are the most significant numbers when ordering two engine versions (that is, a 4.12.* version is always 
+// newer than a 4.11.* version, regardless of the changelist that it was built with)
 #define ENGINE_MAJOR_VERSION	4
 #define ENGINE_MINOR_VERSION	14
 #define ENGINE_PATCH_VERSION	0
 
+// If set to 1, indicates that this is a licensee build of the engine. For the same major/minor/patch release of the engine, licensee changelists are always 
+// considered newer than Epic changelists for engine versions. This follows the assumption that content is developed by Epic leading up to a release, and which 
+// point we lock compatibility, and any subsequent licensee modifications to the engine will have a superset of its functionality even if the changelist 
+// numbers are lower.
 #define ENGINE_IS_LICENSEE_VERSION 0
 
+// The Perforce changelist being compiled. Use this value advisedly; it does not take into account out-of-order commits to engine release branches over 
+// development branches, licensee versions, or whether the engine version has been locked to maintain compatibility with a previous engine release. Prefer
+// BUILD_VERSION where a unique, product-specific identifier is required, or FEngineVersion::CompatibleWith() where relational comparisons between two 
+// versions is required.
 #define BUILT_FROM_CHANGELIST 0
-#define BRANCH_NAME "UE4"
-
-#define ENGINE_IS_PROMOTED_BUILD (BUILT_FROM_CHANGELIST > 0)
-
-#define EPIC_COMPANY_NAME  "Epic Games, Inc."
-#define EPIC_COPYRIGHT_STRING "Copyright 1998-2016 Epic Games, Inc. All Rights Reserved."
-#define EPIC_PRODUCT_NAME "Unreal Engine"
-#define EPIC_PRODUCT_IDENTIFIER "UnrealEngine"
-
-#define BUILD_VERSION VERSION_TEXT(BRANCH_NAME) VERSION_TEXT("-CL-") VERSION_STRINGIFY(BUILT_FROM_CHANGELIST) 
-
-#endif // BUILD_VERSION_LOCKED
 
-#if LOCKED_NETWORK_VERSION
-	#define ENGINE_NET_VERSION LOCKED_NETWORK_VERSION
+// Whether this build is "promoted"; that is, compiled by a build machine (rather than locally) and distributed in binary form. This disables certain features in
+// the engine relating to building locally (because they require intermediate files to be available), such as the hot-reload functionality in the editor.
+// UnrealGameSync explicitly sets this to zero for local builds.
+#define ENGINE_IS_PROMOTED_BUILD  (BUILT_FROM_CHANGELIST > 0)
+
+// The changelist version of the engine. By including a monotonically increasing number in the engine version (and saving it into packages as FEngineVersion), 
+// we can prevent newer packages with the same major/minor engine version from being loaded with an older revision of the engine.
+#define ENGINE_CURRENT_CL_VERSION  BUILT_FROM_CHANGELIST
+
+// The compatible changelist version of the engine. This number identifies a particular API revision, and is used to determine module and package backwards 
+// compatibility. Hotfixes should retain the compatible version of the original release. This define is parsed by the build tools, and should be a number or 
+// BUILT_FROM_CHANGELIST, defined in this particular order for each alternative.
+#if ENGINE_CURRENT_CL_VERSION > 0
+	#if ENGINE_IS_LICENSEE_VERSION
+		#define ENGINE_COMPATIBLE_CL_VERSION ENGINE_CURRENT_CL_VERSION
+	#else
+		#define ENGINE_COMPATIBLE_CL_VERSION ENGINE_CURRENT_CL_VERSION /* Or hotfix compatibility changelist */
+	#endif
 #else
-	// if not defined default network to build cl
-	#define ENGINE_NET_VERSION BUILT_FROM_CHANGELIST
+	#define ENGINE_COMPATIBLE_CL_VERSION 0
 #endif
-
-#if LOCKED_REPLAY_VERSION
-	#define ENGINE_REPLAY_VERSION LOCKED_REPLAY_VERSION
-#else
-	// If not defined default replays to network version
-	#define ENGINE_REPLAY_VERSION ENGINE_NET_VERSION
+#if ENGINE_IS_LICENSEE_VERSION == 0 && ENGINE_PATCH_VERSION > 0 && ENGINE_CL_VERSION > 0 && ENGINE_COMPATIBLE_CL_VERSION == ENGINE_CL_VERSION
+	#error ENGINE_COMPATIBLE_CL_VERSION must be manually defined for hotfix builds
 #endif
 
-// Checks
+// The version number used for determining network compatibility
+#define ENGINE_NET_VERSION  ENGINE_COMPATIBLE_CL_VERSION
 
-#ifndef BUILT_FROM_CHANGELIST
-	#error BUILT_FROM_CHANGELIST must be defined!
-#endif
+// The version number used for determining replay compatibility
+#define ENGINE_REPLAY_VERSION  ENGINE_NET_VERSION
 
-#ifndef ENGINE_NET_VERSION
-	#error ENGINE_NET_VERSION must be defined!
-#endif
+// The branch that this engine is being built from. When set by UAT, this has the form of a Perforce depot path with forward slashes escaped by plus characters 
+// (eg. //UE4/Main -> ++UE4+Main)
+#define BRANCH_NAME "UE4"
 
-#ifndef ENGINE_REPLAY_VERSION
-	#error ENGINE_REPLAY_VERSION must be defined!
-#endif
+// Macros for encoding strings
+#define VERSION_TEXT_2(x) L ## x
+#define VERSION_TEXT(x) VERSION_TEXT_2(x)
+#define VERSION_STRINGIFY_2(x) VERSION_TEXT(#x)
+#define VERSION_STRINGIFY(x) VERSION_STRINGIFY_2(x)
 
-#ifndef BUILD_VERSION
-	#error BUILD_VERSION must be defined!
-#endif
+// The BUILD_VERSION string defines an opaque string representing this particular build, and can be customized for a product without modifying the internal engine version.
+// This string should be used to uniquely identify a build of the current product, as opposed to something built with this version of the engine.
+#define BUILD_VERSION VERSION_TEXT(BRANCH_NAME) VERSION_TEXT("-CL-") VERSION_STRINGIFY(BUILT_FROM_CHANGELIST) 
+
+// Various strings used for resources
+#define EPIC_COMPANY_NAME  "Epic Games, Inc."
+#define EPIC_COPYRIGHT_STRING "Copyright 1998-2016 Epic Games, Inc. All Rights Reserved."
+#define EPIC_PRODUCT_NAME "Unreal Engine"
+#define EPIC_PRODUCT_IDENTIFIER "UnrealEngine"
 
 #define ENGINE_VERSION_STRING \
 	VERSION_STRINGIFY(ENGINE_MAJOR_VERSION) \
diff --git a/Engine/Source/Runtime/Launch/Resources/VersionLocked.h b/Engine/Source/Runtime/Launch/Resources/VersionLocked.h
deleted file mode 100644
diff --git a/Engine/Source/Runtime/Online/BuildPatchServices/Private/BuildPatchChunk.h b/Engine/Source/Runtime/Online/BuildPatchServices/Private/BuildPatchChunk.h
--- a/Engine/Source/Runtime/Online/BuildPatchServices/Private/BuildPatchChunk.h
+++ b/Engine/Source/Runtime/Online/BuildPatchServices/Private/BuildPatchChunk.h
@@ -6,6 +6,8 @@
 
 #pragma once
 
+#include "CoreUObject.h"
+
 #include "Generation/StatsCollector.h"
 
 #include "BuildPatchChunk.generated.h"
