Update version in //UE4/Main to 4.15.

diff --git a/Engine/Build/Build.version b/Engine/Build/Build.version
--- a/Engine/Build/Build.version
+++ b/Engine/Build/Build.version
@@ -1,6 +1,6 @@
 {
 	"MajorVersion": 4,
-	"MinorVersion": 14,
+	"MinorVersion": 15,
 	"PatchVersion": 0,
 	"Changelist": 0,
 	"CompatibleChangelist": 0,
diff --git a/Engine/Source/Runtime/Launch/Resources/Version.h b/Engine/Source/Runtime/Launch/Resources/Version.h
--- a/Engine/Source/Runtime/Launch/Resources/Version.h
+++ b/Engine/Source/Runtime/Launch/Resources/Version.h
@@ -55,7 +55,7 @@
 // These numbers define the banner UE4 version, and are the most significant numbers when ordering two engine versions (that is, a 4.12.* version is always 
 // newer than a 4.11.* version, regardless of the changelist that it was built with)
 #define ENGINE_MAJOR_VERSION	4
-#define ENGINE_MINOR_VERSION	14
+#define ENGINE_MINOR_VERSION	15
 #define ENGINE_PATCH_VERSION	0
 
 // If set to 1, indicates that this is a licensee build of the engine. For the same major/minor/patch release of the engine, licensee changelists are always 
