Branch snapshot for CL 3210543

[CL 3210543 in Main branch]

diff --git a/Engine/Source/Runtime/Launch/Resources/Windows/PCLaunch.rc2 b/Engine/Source/Runtime/Launch/Resources/Windows/PCLaunch.rc2
new file mode 100644
--- /dev/null
+++ b/Engine/Source/Runtime/Launch/Resources/Windows/PCLaunch.rc2
@@ -0,0 +1,70 @@
+// Copyright 1998-2016 Epic Games, Inc. All Rights Reserved.
+//
+// This file contains resources which are edited by hand, and cannot be loaded/saved by the Visual Studio resource editor.
+
+#include "resource.h"
+#include "../../../Core/Public/Misc/CoreMiscDefines.h"
+#include "../Version.h"
+
+/////////////////////////////////////////////////////////////////////////////
+//
+// Version
+//
+
+VS_VERSION_INFO VERSIONINFO
+ FILEVERSION ENGINE_MAJOR_VERSION,ENGINE_MINOR_VERSION,ENGINE_PATCH_VERSION,0
+ PRODUCTVERSION ENGINE_MAJOR_VERSION,ENGINE_MINOR_VERSION,ENGINE_PATCH_VERSION,0
+ FILEFLAGSMASK 0x17L
+#ifdef _DEBUG
+ FILEFLAGS 0x1L
+#else
+ FILEFLAGS 0x0L
+#endif
+ FILEOS 0x4L
+ FILETYPE 0x2L
+ FILESUBTYPE 0x0L
+BEGIN
+	BLOCK "StringFileInfo"
+	BEGIN
+		BLOCK "040904b0"
+		BEGIN
+			VALUE "CompanyName", BUILD_PROJECT_COMPANY_NAME
+			VALUE "LegalCopyright", BUILD_PROJECT_COPYRIGHT_STRING
+			VALUE "ProductName", BUILD_PROJECT_PRODUCT_NAME
+			VALUE "ProductVersion", BUILD_VERSION
+			VALUE "FileDescription", BUILD_PROJECT_PRODUCT_NAME
+			VALUE "InternalName", BUILD_PROJECT_PRODUCT_IDENTIFIER
+			VALUE "OriginalFilename", PREPROCESSOR_TO_STRING(ORIGINAL_FILE_NAME)
+		END
+	END
+	BLOCK "VarFileInfo"
+	BEGIN
+		VALUE "Translation", 0x409, 1200
+	END
+END
+
+/////////////////////////////////////////////////////////////////////////////
+//
+// Application Manifest
+//
+
+#if _DEBUG
+
+// Debug manifest
+#if _WIN64
+1						RT_MANIFEST				"Debug-PCLaunch-Win64.manifest"
+#else
+1						RT_MANIFEST				"Debug-PCLaunch-Win32.manifest"
+#endif
+
+#else
+
+// Release manifest
+#if _WIN64
+1						RT_MANIFEST				"PCLaunch-Win64.manifest"
+#else
+1						RT_MANIFEST				"PCLaunch-Win32.manifest"
+#endif
+
+#endif
+
