Copying //UE4/GamePlugins-Staging to //UE4/Dev-Main (Source: //GamePlugins/Main @ 3232656)

diff --git a/Engine/Plugins/Online/OnlineSubsystem/Source/Public/Interfaces/OnlineStoreInterface.h b/Engine/Plugins/Online/OnlineSubsystem/Source/Public/Interfaces/OnlineStoreInterface.h
--- a/Engine/Plugins/Online/OnlineSubsystem/Source/Public/Interfaces/OnlineStoreInterface.h
+++ b/Engine/Plugins/Online/OnlineSubsystem/Source/Public/Interfaces/OnlineStoreInterface.h
@@ -83,8 +83,8 @@ struct FInAppPurchaseProductRequest
 
 
 /**
-* Micro-transaction purchase information
-*/
+ * Micro-transaction purchase information
+ */
 USTRUCT(BlueprintType)
 struct FInAppPurchaseProductInfo
 {
@@ -130,15 +130,15 @@ struct FInAppPurchaseProductInfo
 	UPROPERTY(BlueprintReadOnly, Category = ProductInfo)
 	FString GroupingSeparator;
 
-	// The localized display price name
+	// Opaque receipt data for the transaction
 	UPROPERTY(BlueprintReadOnly, Category = ProductInfo)
 	FString ReceiptData;
 };
 
 
 /**
-* Micro-transaction restored purchase information
-*/
+ * Micro-transaction restored purchase information
+ */
 USTRUCT(BlueprintType)
 struct FInAppPurchaseRestoreInfo
 {
@@ -148,9 +148,9 @@ struct FInAppPurchaseRestoreInfo
 	UPROPERTY(BlueprintReadOnly, Category = ProductInfo)
 	FString Identifier;
 
-	// The localized display price name
+	// The opaque receipt data for the platform
 	UPROPERTY(BlueprintReadOnly, Category = ProductInfo)
-		FString ReceiptData;
+	FString ReceiptData;
 
 	// the unique transaction identifier
 	UPROPERTY(BlueprintReadOnly, Category = ProductInfo)
@@ -159,8 +159,8 @@ struct FInAppPurchaseRestoreInfo
 
 
 /**
-*	Interface for reading data from an In App Purchase service
-*/
+ *	Interface for reading data from an In App Purchase service
+ */
 class ONLINESUBSYSTEM_API FOnlineProductInformationRead
 {
 public:
@@ -187,9 +187,12 @@ class ONLINESUBSYSTEM_API FOnlineInAppPurchaseTransaction
 public:
 	/** Indicates an error reading data occurred while processing */
 	EOnlineAsyncTaskState::Type ReadState;
+	/** Is the purchase consumable */
+	bool bIsConsumable;
 
 	FOnlineInAppPurchaseTransaction() :
-		ReadState(EOnlineAsyncTaskState::NotStarted)
+		ReadState(EOnlineAsyncTaskState::NotStarted),
+		bIsConsumable(false)
 	{
 	}
 
@@ -201,8 +204,8 @@ typedef TSharedPtr<FOnlineInAppPurchaseTransaction, ESPMode::ThreadSafe> FOnline
 
 
 /**
-*	Interface for reading data from an In App Purchase service
-*/
+ *	Interface for reading data from an In App Purchase service
+ */
 class ONLINESUBSYSTEM_API FOnlineInAppPurchaseRestoreRead
 {
 public:
diff --git a/Engine/Plugins/Online/OnlineSubsystem/Source/Public/Interfaces/OnlineUserInterface.h b/Engine/Plugins/Online/OnlineSubsystem/Source/Public/Interfaces/OnlineUserInterface.h
--- a/Engine/Plugins/Online/OnlineSubsystem/Source/Public/Interfaces/OnlineUserInterface.h
+++ b/Engine/Plugins/Online/OnlineSubsystem/Source/Public/Interfaces/OnlineUserInterface.h
@@ -56,7 +56,7 @@ public:
 	DEFINE_ONLINE_PLAYER_DELEGATE_THREE_PARAM(MAX_LOCAL_PLAYERS, OnQueryUserInfoComplete, bool, const TArray< TSharedRef<const FUniqueNetId> >&, const FString&);
 
 	/**
-	 * Obtains the cached list of online user info 
+	 * Obtains the cached list of online user info
 	 *
 	 * @param LocalUserNum the local user that queried for online user data
 	 * @param OutUsers [out] array that receives the copied data
@@ -101,7 +101,7 @@ public:
 	 *
 	 * @param bWasSuccessful true if server was contacted and a valid result received
 	 * @param UserId user id initiating the request
-	 * @param AuthType auth type that the external ids represent
+	 * @param QueryOptions Options specifying how to treat the External IDs and other query-related settings
 	 * @param ExternalIds array of external ids to attempt to map to user ids
 	 * @param Error string representing the error condition
 	 */
@@ -111,7 +111,7 @@ public:
 	 * Contacts server to obtain user ids from external ids
 	 *
 	 * @param UserId id of the user that is requesting the name string lookup
-	 * @param AuthType auth type that the external ids represent
+	 * @param QueryOptions Options specifying how to treat the External IDs and other query-related settings
 	 * @param ExternalIds array of external ids to attempt to map to user ids
 	 *
 	 * @return true if the operation was started successfully
@@ -121,7 +121,7 @@ public:
 	/**
 	 * Get the cached user ids for the specified external ids
 	 *
-	 * @param AuthType auth type that the external ids represent
+	 * @param QueryOptions Options specifying how to treat the External IDs and other query-related settings
 	 * @param ExternalIds array of external ids to map to user ids
 	 * @param OutIds array of user ids that map to the specified external ids (can contain null entries)
 	 */
@@ -130,7 +130,7 @@ public:
 	/**
 	 * Get the cached user id for the specified external id
 	 *
-	 * @param AuthType auth type that the external ids represent
+	 * @param QueryOptions Options specifying how to treat the External IDs and other query-related settings
 	 * @param ExternalId external id to obtain user id for
 
 	 * @return user info or null ptr if not found
diff --git a/Engine/Plugins/Online/OnlineSubsystem/Source/Public/Interfaces/VoiceInterface.h b/Engine/Plugins/Online/OnlineSubsystem/Source/Public/Interfaces/VoiceInterface.h
--- a/Engine/Plugins/Online/OnlineSubsystem/Source/Public/Interfaces/VoiceInterface.h
+++ b/Engine/Plugins/Online/OnlineSubsystem/Source/Public/Interfaces/VoiceInterface.h
@@ -340,7 +340,7 @@ public:
 	 * @param LocalUserNum the controller number of the associated user
 	 * @param UniqueId the id of the player being checked
 	 *
-	 * @return true if a member of their friends list, false otherwise
+	 * @return true if the specified user is muted, false otherwise
 	 */
 	virtual bool IsMuted(uint32 LocalUserNum, const FUniqueNetId& UniqueId) const = 0;
 
diff --git a/Engine/Plugins/Online/OnlineSubsystem/Source/Public/OnlineAsyncTaskManager.h b/Engine/Plugins/Online/OnlineSubsystem/Source/Public/OnlineAsyncTaskManager.h
--- a/Engine/Plugins/Online/OnlineSubsystem/Source/Public/OnlineAsyncTaskManager.h
+++ b/Engine/Plugins/Online/OnlineSubsystem/Source/Public/OnlineAsyncTaskManager.h
@@ -222,7 +222,7 @@ PACKAGE_SCOPE:
 	volatile bool bWasSuccessful;
 
 	/** Hidden on purpose */
-	FOnlineAsyncTaskBasic() 
+	FOnlineAsyncTaskBasic()
 		: Subsystem(NULL)
 		, bIsComplete(false)
 		, bWasSuccessful(false)
@@ -231,7 +231,7 @@ PACKAGE_SCOPE:
 
 public:
 
-	FOnlineAsyncTaskBasic(T* InSubsystem) 
+	FOnlineAsyncTaskBasic(T* const InSubsystem)
 		: Subsystem(InSubsystem)
 		, bIsComplete(false)
 		, bWasSuccessful(false)
diff --git a/Engine/Plugins/Online/OnlineSubsystem/Source/Public/OnlineError.h b/Engine/Plugins/Online/OnlineSubsystem/Source/Public/OnlineError.h
--- a/Engine/Plugins/Online/OnlineSubsystem/Source/Public/OnlineError.h
+++ b/Engine/Plugins/Online/OnlineSubsystem/Source/Public/OnlineError.h
@@ -5,7 +5,7 @@
 #include "CoreMinimal.h"
 
 namespace EOnlineServerConnectionStatus {
-	enum Type : unsigned int;
+	enum Type : uint8;
 }
 
 /** Generic Error response for OSS calls */
diff --git a/Engine/Plugins/Online/OnlineSubsystem/Source/Public/OnlineSessionSettings.h b/Engine/Plugins/Online/OnlineSubsystem/Source/Public/OnlineSessionSettings.h
--- a/Engine/Plugins/Online/OnlineSubsystem/Source/Public/OnlineSessionSettings.h
+++ b/Engine/Plugins/Online/OnlineSubsystem/Source/Public/OnlineSessionSettings.h
@@ -8,6 +8,9 @@
 #include "OnlineKeyValuePair.h"
 #include "OnlineSubsystemPackage.h"
 
+/** default beacon port, if not specified by other means */
+#define DEFAULT_BEACON_PORT 15000
+
 /** Setting describing the name of the current map (value is FString) */
 #define SETTING_MAPNAME FName(TEXT("MAPNAME"))
 /** Setting describing the number of bots in the session (value is int32) */
@@ -603,7 +606,8 @@ public:
 	 * in the same bucket is often not a useful comparison and skill is better
 	 */
 	int32 PingBucketSize;
-
+	/** Search hash used by the online subsystem to disambiguate search queries, stamped every time FindSession is called */
+	int32 PlatformHash;
 	/** Amount of time to wait for the search results. May not apply to all platforms. */
 	float TimeoutInSeconds;
 
@@ -613,6 +617,7 @@ public:
 		MaxSearchResults(1),
 		bIsLanQuery(false),
 		PingBucketSize(0),
+		PlatformHash(0),
 		TimeoutInSeconds(0.0f)
 	{
 		QuerySettings.Set(SETTING_MAPNAME, FString(TEXT("")), EOnlineComparisonOp::Equals);
diff --git a/Engine/Plugins/Online/OnlineSubsystem/Source/Public/OnlineSubsystem.h b/Engine/Plugins/Online/OnlineSubsystem/Source/Public/OnlineSubsystem.h
--- a/Engine/Plugins/Online/OnlineSubsystem/Source/Public/OnlineSubsystem.h
+++ b/Engine/Plugins/Online/OnlineSubsystem/Source/Public/OnlineSubsystem.h
@@ -472,6 +472,11 @@ public:
 	 */
 	virtual bool Init() = 0;
 
+	/** 
+	 * Perform any shutdown actions prior to any other modules being unloaded/shutdown
+	 */
+	virtual void PreUnload() = 0;
+
 	/** 
 	 * Shutdown the underlying subsystem APIs
 	 * @return true if the subsystem shutdown successfully, false otherwise
diff --git a/Engine/Plugins/Online/OnlineSubsystem/Source/Public/OnlineSubsystemImpl.h b/Engine/Plugins/Online/OnlineSubsystem/Source/Public/OnlineSubsystemImpl.h
--- a/Engine/Plugins/Online/OnlineSubsystem/Source/Public/OnlineSubsystemImpl.h
+++ b/Engine/Plugins/Online/OnlineSubsystem/Source/Public/OnlineSubsystemImpl.h
@@ -8,6 +8,8 @@
 #include "Containers/Queue.h"
 #include "OnlineSubsystemPackage.h"
 
+struct FOnlineError;
+
 DECLARE_DELEGATE(FNextTickDelegate);
 
 namespace OSSConsoleVariables
@@ -28,6 +30,13 @@ private:
 	 */
 	bool HandleFriendExecCommands(UWorld* InWorld, const TCHAR* Cmd, FOutputDevice& Ar);
 	bool HandleSessionExecCommands(UWorld* InWorld, const TCHAR* Cmd, FOutputDevice& Ar);
+	bool HandlePurchaseExecCommands(UWorld* InWorld, const TCHAR* Cmd, FOutputDevice& Ar);
+	
+	/** Delegate fired when exec cheat related to receipts completes */
+	void OnQueryReceiptsComplete(const FOnlineError& Result, TSharedPtr<const FUniqueNetId> UserId);
+	
+	/** Dump purchase receipts for a given user id */
+	void DumpReceipts(const FUniqueNetId& UserId);
 
 protected:
 
@@ -70,6 +79,7 @@ public:
 	virtual ~FOnlineSubsystemImpl();
 
 	// IOnlineSubsystem
+	virtual void PreUnload() override;
 	virtual bool Shutdown() override;
 	virtual bool IsServer() const override;
 	virtual bool IsDedicated() const override{ return bForceDedicated || IsRunningDedicatedServer(); }
diff --git a/Engine/Plugins/Online/OnlineSubsystem/Source/Public/OnlineSubsystemModule.h b/Engine/Plugins/Online/OnlineSubsystem/Source/Public/OnlineSubsystemModule.h
--- a/Engine/Plugins/Online/OnlineSubsystem/Source/Public/OnlineSubsystemModule.h
+++ b/Engine/Plugins/Online/OnlineSubsystem/Source/Public/OnlineSubsystemModule.h
@@ -57,6 +57,11 @@ private:
 	 */
 	void LoadDefaultSubsystem();
 
+	/**
+	 *	Called before ShutdownOnlineSubsystem, before other modules have been unloaded
+	 */
+	virtual void PreUnloadOnlineSubsystem();
+
 	/**
 	 *	Shuts down all registered online subsystem platforms and unloads their modules
 	 */
@@ -135,6 +140,12 @@ public:
 	 */
 	virtual void StartupModule() override;
 
+	/**
+	 * Called before the module has been unloaded
+	 * Overloaded to allow online subsystems to cancel any outstanding http requests
+	 */
+	virtual void PreUnloadCallback() override;
+
 	/**
 	 * Called before the module is unloaded, right before the module object is destroyed.
 	 * Overloaded to shut down all loaded online subsystems
diff --git a/Engine/Plugins/Online/OnlineSubsystem/Source/Public/OnlineSubsystemTypes.h b/Engine/Plugins/Online/OnlineSubsystem/Source/Public/OnlineSubsystemTypes.h
--- a/Engine/Plugins/Online/OnlineSubsystem/Source/Public/OnlineSubsystemTypes.h
+++ b/Engine/Plugins/Online/OnlineSubsystem/Source/Public/OnlineSubsystemTypes.h
@@ -126,7 +126,7 @@ namespace ELoginStatus
 /** Possible connection states */
 namespace EOnlineServerConnectionStatus
 {
-	enum Type : unsigned int
+	enum Type : uint8
 	{
 		/** System normal (used for default state) */
 		Normal = 0,
@@ -811,38 +811,100 @@ public:
 	/** Holds the net id for a player */
 	FString UniqueNetIdStr;
 
+#if PLATFORM_COMPILER_HAS_DEFAULTED_FUNCTIONS
+	FUniqueNetIdString() = default;
+	virtual ~FUniqueNetIdString() = default;
+	FUniqueNetIdString(const FUniqueNetIdString&) = default;
+	FUniqueNetIdString(FUniqueNetIdString&&) = default;
+	FUniqueNetIdString& operator=(const FUniqueNetIdString&) = default;
+	FUniqueNetIdString& operator=(FUniqueNetIdString&&) = default;
+#else
 	/** Default constructor */
 	FUniqueNetIdString()
 	{
 	}
 
+	/** Destructor */
+	virtual ~FUniqueNetIdString()
+	{
+	}
+
+	/**
+	 * Copy Constructor
+	 *
+	 * @param Src the id to copy
+	 */
+	explicit FUniqueNetIdString(const FUniqueNetIdString& Src)
+		: UniqueNetIdStr(Src.UniqueNetIdStr)
+	{
+	}
+
+	/**
+	 * Move Constructor
+	 *
+	 * @param Src the id to copy
+	 */
+	explicit FUniqueNetIdString(FUniqueNetIdString&& Src)
+		: UniqueNetIdStr(MoveTemp(Src.UniqueNetIdStr))
+	{
+	}
+
+	/**
+	 * Copy Assignment Operator
+	 *
+	 * @param Src the id to copy
+	 */
+	FUniqueNetIdString& operator=(const FUniqueNetIdString& Src)
+	{
+		if (this != &Src)
+		{
+			UniqueNetIdStr = Src.UniqueNetIdStr;
+		}
+		return *this;
+	}
+
+	/**
+	 * Move Assignment Operator
+	 *
+	 * @param Src the id to copy
+	 */
+	FUniqueNetIdString& operator=(FUniqueNetIdString&& Src)
+	{
+		if (this != &Src)
+		{
+			UniqueNetIdStr = MoveTemp(Src.UniqueNetIdStr);
+		}
+		return *this;
+	}
+#endif
+
 	/**
 	 * Constructs this object with the specified net id
 	 *
 	 * @param InUniqueNetId the id to set ours to
 	 */
-	explicit FUniqueNetIdString(const FString& InUniqueNetId) 
+	explicit FUniqueNetIdString(const FString& InUniqueNetId)
 		: UniqueNetIdStr(InUniqueNetId)
 	{
 	}
 
 	/**
-	 * Copy Constructor
+	 * Constructs this object with the specified net id
 	 *
-	 * @param Src the id to copy
+	 * @param InUniqueNetId the id to set ours to
 	 */
-	explicit FUniqueNetIdString(const FUniqueNetId& Src) 
-		: UniqueNetIdStr(Src.ToString())
+	explicit FUniqueNetIdString(FString&& InUniqueNetId)
+		: UniqueNetIdStr(MoveTemp(InUniqueNetId))
 	{
 	}
 
 	/**
-	 * Copy Constructor
+	 * Constructs this object with the string value of the specified net id
 	 *
 	 * @param Src the id to copy
 	 */
-	explicit FUniqueNetIdString(const FUniqueNetIdString& Src) 
-		: UniqueNetIdStr(Src.UniqueNetIdStr)
+	explicit FUniqueNetIdString(const FUniqueNetId& Src)
+		: UniqueNetIdStr(Src.ToString())
 	{
 	}
 
diff --git a/Engine/Plugins/Online/OnlineSubsystemAmazon/Source/OnlineSubsystemAmazon.Build.cs b/Engine/Plugins/Online/OnlineSubsystemAmazon/Source/OnlineSubsystemAmazon.Build.cs
--- a/Engine/Plugins/Online/OnlineSubsystemAmazon/Source/OnlineSubsystemAmazon.Build.cs
+++ b/Engine/Plugins/Online/OnlineSubsystemAmazon/Source/OnlineSubsystemAmazon.Build.cs
@@ -7,7 +7,8 @@ public class OnlineSubsystemAmazon : ModuleRules
 {
 	public OnlineSubsystemAmazon(TargetInfo Target)
 	{
-		Definitions.Add("ONLINESUBSYSTEMAMAZON_PACKAGE=1");		
+		Definitions.Add("ONLINESUBSYSTEMAMAZON_PACKAGE=1");	
+		PCHUsage = ModuleRules.PCHUsageMode.UseExplicitOrSharedPCHs;
 
 		PrivateDependencyModuleNames.AddRange(
 			new string[] {
diff --git a/Engine/Plugins/Online/OnlineSubsystemFacebook/Source/Private/Windows/OnlineFriendsFacebook.cpp b/Engine/Plugins/Online/OnlineSubsystemFacebook/Source/Private/Windows/OnlineFriendsFacebook.cpp
--- a/Engine/Plugins/Online/OnlineSubsystemFacebook/Source/Private/Windows/OnlineFriendsFacebook.cpp
+++ b/Engine/Plugins/Online/OnlineSubsystemFacebook/Source/Private/Windows/OnlineFriendsFacebook.cpp
@@ -55,9 +55,11 @@ FOnlineFriendsFacebook::FOnlineFriendsFacebook(FOnlineSubsystemFacebook* InSubsy
 		UE_LOG(LogOnline, Warning, TEXT("Missing LoginUrl= in [OnlineSubsystemFacebook.OnlineIdentityFacebook] of DefaultEngine.ini"));
 	}
 	GConfig->GetArray(TEXT("OnlineSubsystemFacebook.OnlineFriendsFacebook"), TEXT("FriendsFields"), FriendsFields, GEngineIni);	
+
+	// @HSL_BEGIN - Josh.May - 10/04/2016 - Removed deprecated 'username' parameter
 	// always required fields
 	FriendsFields.AddUnique(TEXT("name"));
-	FriendsFields.AddUnique(TEXT("username"));
+	// @HSL_END - Josh.May - 10/04/2016
 }
 
 FOnlineFriendsFacebook::~FOnlineFriendsFacebook()
diff --git a/Engine/Plugins/Online/OnlineSubsystemFacebook/Source/Private/Windows/OnlineIdentityFacebook.cpp b/Engine/Plugins/Online/OnlineSubsystemFacebook/Source/Private/Windows/OnlineIdentityFacebook.cpp
--- a/Engine/Plugins/Online/OnlineSubsystemFacebook/Source/Private/Windows/OnlineIdentityFacebook.cpp
+++ b/Engine/Plugins/Online/OnlineSubsystemFacebook/Source/Private/Windows/OnlineIdentityFacebook.cpp
@@ -75,6 +75,12 @@ FOnlineIdentityFacebook::FOnlineIdentityFacebook()
 		// Default to 30 seconds
 		MaxCheckElapsedTime = 30.f;
 	}
+
+	// @HSL_BEGIN - Josh.May - 10/04/2016 - Added permission scope fields
+	GConfig->GetArray(TEXT("OnlineSubsystemFacebook.OnlineIdentityFacebook"), TEXT("ScopeFields"), ScopeFields, GEngineIni);
+	// always required fields
+	ScopeFields.AddUnique(TEXT("public_profile"));
+	// @HSL_END - Josh.May - 10/04/2016
 }
 
 /**
@@ -245,8 +251,14 @@ bool FOnlineIdentityFacebook::Login(int32 LocalUserNum, const FOnlineAccountCred
 		// random number to represent client generated state for verification on login
 		State = FString::FromInt(FMath::Rand() % 100000);
 		// auth url to spawn in browser
-		const FString& Command = FString::Printf(TEXT("%s?redirect_uri=%s&client_id=%s&state=%s&response_type=token"), 
-			*LoginUrl, *LoginRedirectUrl, *ClientId, *State);
+
+		// @HSL_BEGIN - Josh.May - 10/04/2016 - Added permission scope fields
+		FString Scopes = FString::Join(ScopeFields, TEXT(","));
+
+		const FString& Command = FString::Printf(TEXT("%s?redirect_uri=%s&client_id=%s&state=%s&response_type=token&scope=%s"), 
+			*LoginUrl, *LoginRedirectUrl, *ClientId, *State, *Scopes);
+		// @HSL_END - Josh.May - 10/04/2016
+
 		// This should open the browser with the command as the URL
 		if (FPlatformMisc::OsExecute(TEXT("open"), *Command))
 		{
diff --git a/Engine/Plugins/Online/OnlineSubsystemFacebook/Source/Private/Windows/OnlineIdentityFacebook.h b/Engine/Plugins/Online/OnlineSubsystemFacebook/Source/Private/Windows/OnlineIdentityFacebook.h
--- a/Engine/Plugins/Online/OnlineSubsystemFacebook/Source/Private/Windows/OnlineIdentityFacebook.h
+++ b/Engine/Plugins/Online/OnlineSubsystemFacebook/Source/Private/Windows/OnlineIdentityFacebook.h
@@ -82,6 +82,11 @@ class FOnlineIdentityFacebook :
 	/** The client id given to us by Facebook */
 	FString ClientId;
 
+	// @HSL_BEGIN - Josh.May - 10/04/2016 - Added permission scope fields
+	/** Config based list of permission scopes to use when logging in */
+	TArray<FString> ScopeFields;
+	// @HSL_END - Josh.May - 10/04/2016
+
 	/** Users that have been registered/authenticated */
 	FUserOnlineAccountFacebookMap UserAccounts;
 	/** Ids mapped to locally registered users */
diff --git a/Engine/Plugins/Online/OnlineSubsystemNull/Source/Private/OnlineSessionInterfaceNull.cpp b/Engine/Plugins/Online/OnlineSubsystemNull/Source/Private/OnlineSessionInterfaceNull.cpp
--- a/Engine/Plugins/Online/OnlineSubsystemNull/Source/Private/OnlineSessionInterfaceNull.cpp
+++ b/Engine/Plugins/Online/OnlineSubsystemNull/Source/Private/OnlineSessionInterfaceNull.cpp
@@ -735,7 +735,7 @@ bool FOnlineSessionNull::GetResolvedConnectString(FName SessionName, FString& Co
 	return bSuccess;
 }
 
-bool FOnlineSessionNull::GetResolvedConnectString(const class FOnlineSessionSearchResult& SearchResult, FName PortType, FString& ConnectInfo)
+bool FOnlineSessionNull::GetResolvedConnectString(const FOnlineSessionSearchResult& SearchResult, FName PortType, FString& ConnectInfo)
 {
 	bool bSuccess = false;
 	if (SearchResult.Session.SessionInfo.IsValid())
@@ -744,11 +744,11 @@ bool FOnlineSessionNull::GetResolvedConnectString(const class FOnlineSessionSear
 
 		if (PortType == BeaconPort)
 		{
-			int32 BeaconListenPort = 15000;
+			int32 BeaconListenPort = DEFAULT_BEACON_PORT;
 			if (!SearchResult.Session.SessionSettings.Get(SETTING_BEACONPORT, BeaconListenPort) || BeaconListenPort <= 0)
 			{
-				// Reset the default BeaconListenPort back to 15000 because the SessionSettings value does not exist or was not valid
-				BeaconListenPort = 15000;
+				// Reset the default BeaconListenPort back to DEFAULT_BEACON_PORT because the SessionSettings value does not exist or was not valid
+				BeaconListenPort = DEFAULT_BEACON_PORT;
 			}
 			bSuccess = GetConnectStringFromSessionInfo(SessionInfo, ConnectInfo, BeaconListenPort);
 
diff --git a/Engine/Plugins/Online/OnlineSubsystemSteam/Source/Private/OnlineExternalUIInterfaceSteam.cpp b/Engine/Plugins/Online/OnlineSubsystemSteam/Source/Private/OnlineExternalUIInterfaceSteam.cpp
--- a/Engine/Plugins/Online/OnlineSubsystemSteam/Source/Private/OnlineExternalUIInterfaceSteam.cpp
+++ b/Engine/Plugins/Online/OnlineSubsystemSteam/Source/Private/OnlineExternalUIInterfaceSteam.cpp
@@ -46,7 +46,7 @@ bool FOnlineExternalUISteam::ShowFriendsUI(int32 LocalUserNum)
 	return true;
 }
 
-bool FOnlineExternalUISteam::ShowInviteUI(int32 LocalUserNum, FName SessionMame)
+bool FOnlineExternalUISteam::ShowInviteUI(int32 LocalUserNum, FName SessionName)
 {
 	IOnlineSessionPtr SessionInt = SteamSubsystem->GetSessionInterface();
 	if (SessionInt.IsValid() && SessionInt->HasPresenceSession())
diff --git a/Engine/Plugins/Online/OnlineSubsystemSteam/Source/Private/OnlineExternalUIInterfaceSteam.h b/Engine/Plugins/Online/OnlineSubsystemSteam/Source/Private/OnlineExternalUIInterfaceSteam.h
--- a/Engine/Plugins/Online/OnlineSubsystemSteam/Source/Private/OnlineExternalUIInterfaceSteam.h
+++ b/Engine/Plugins/Online/OnlineSubsystemSteam/Source/Private/OnlineExternalUIInterfaceSteam.h
@@ -82,7 +82,7 @@ public:
 	// IOnlineExternalUI
 	virtual bool ShowLoginUI(const int ControllerIndex, bool bShowOnlineOnly, const FOnLoginUIClosedDelegate& Delegate = FOnLoginUIClosedDelegate()) override;
 	virtual bool ShowFriendsUI(int32 LocalUserNum) override;
-	virtual bool ShowInviteUI(int32 LocalUserNum, FName SessionMame = GameSessionName) override;
+	virtual bool ShowInviteUI(int32 LocalUserNum, FName SessionName = GameSessionName) override;
 	virtual bool ShowAchievementsUI(int32 LocalUserNum) override;
 	virtual bool ShowLeaderboardUI(const FString& LeaderboardName) override;
 	virtual bool ShowWebURL(const FString& Url, const FShowWebUrlParams& ShowParams, const FOnShowWebUrlClosedDelegate& Delegate = FOnShowWebUrlClosedDelegate()) override;
diff --git a/Engine/Plugins/Online/OnlineSubsystemSteam/Source/Private/OnlineSessionInterfaceSteam.cpp b/Engine/Plugins/Online/OnlineSubsystemSteam/Source/Private/OnlineSessionInterfaceSteam.cpp
--- a/Engine/Plugins/Online/OnlineSubsystemSteam/Source/Private/OnlineSessionInterfaceSteam.cpp
+++ b/Engine/Plugins/Online/OnlineSubsystemSteam/Source/Private/OnlineSessionInterfaceSteam.cpp
@@ -1207,7 +1207,7 @@ bool FOnlineSessionSteam::GetResolvedConnectString(FName SessionName, FString& C
 	return bSuccess;
 }
 
-bool FOnlineSessionSteam::GetResolvedConnectString(const class FOnlineSessionSearchResult& SearchResult, const FName PortType, FString& ConnectInfo)
+bool FOnlineSessionSteam::GetResolvedConnectString(const FOnlineSessionSearchResult& SearchResult, const FName PortType, FString& ConnectInfo)
 {
 	bool bSuccess = false;
 	if (SearchResult.Session.SessionInfo.IsValid())
@@ -1216,14 +1216,13 @@ bool FOnlineSessionSteam::GetResolvedConnectString(const class FOnlineSessionSea
 
 		if (PortType == BeaconPort)
 		{
-			int32 BeaconListenPort = 15000;
+			int32 BeaconListenPort = DEFAULT_BEACON_PORT;
 			if (!SearchResult.Session.SessionSettings.Get(SETTING_BEACONPORT, BeaconListenPort) || BeaconListenPort <= 0)
 			{
-				// Reset the default BeaconListenPort back to 15000 because the SessionSettings value does not exist or was not valid
-				BeaconListenPort = 15000;
+				// Reset the default BeaconListenPort back to DEFAULT_BEACON_PORT because the SessionSettings value does not exist or was not valid
+				BeaconListenPort = DEFAULT_BEACON_PORT;
 			}
 			bSuccess = GetConnectStringFromSessionInfo(SessionInfo, ConnectInfo, BeaconListenPort);
-
 		}
 		else if (PortType == GamePort)
 		{
diff --git a/Engine/Plugins/Online/OnlineSubsystemSteam/Source/Private/VoiceInterfaceSteam.cpp b/Engine/Plugins/Online/OnlineSubsystemSteam/Source/Private/VoiceInterfaceSteam.cpp
--- a/Engine/Plugins/Online/OnlineSubsystemSteam/Source/Private/VoiceInterfaceSteam.cpp
+++ b/Engine/Plugins/Online/OnlineSubsystemSteam/Source/Private/VoiceInterfaceSteam.cpp
@@ -408,13 +408,12 @@ bool FOnlineVoiceSteam::IsRemotePlayerTalking(const FUniqueNetId& UniqueId)
 
 bool FOnlineVoiceSteam::IsMuted(uint32 LocalUserNum, const FUniqueNetId& UniqueId) const
 {
-	int32 Index = INDEX_NONE;
 	if (LocalUserNum >= 0 && LocalUserNum < MAX_LOCAL_PLAYERS)
 	{
-		Index = MuteList.Find((const FUniqueNetIdSteam&)UniqueId);
+		return IsLocallyMuted(UniqueId);
 	}
 
-	return Index != INDEX_NONE;
+	return false;
 }
 
 bool FOnlineVoiceSteam::IsLocallyMuted(const FUniqueNetId& UniqueId) const
@@ -423,6 +422,12 @@ bool FOnlineVoiceSteam::IsLocallyMuted(const FUniqueNetId& UniqueId) const
 	return Index != INDEX_NONE;
 }
 
+bool FOnlineVoiceSteam::IsSystemWideMuted(const FUniqueNetId& UniqueId) const
+{
+	int32 Index = SystemMuteList.Find((const FUniqueNetIdSteam&)UniqueId);
+	return Index != INDEX_NONE;
+}
+
 bool FOnlineVoiceSteam::MuteRemoteTalker(uint8 LocalUserNum, const FUniqueNetId& PlayerId, bool bIsSystemWide)
 {
 	uint32 Return = E_FAIL;
@@ -430,7 +435,9 @@ bool FOnlineVoiceSteam::MuteRemoteTalker(uint8 LocalUserNum, const FUniqueNetId&
 	{
 		if (bIsSystemWide)
 		{
+			// Add them to the system wide mute list
 			SystemMuteList.AddUnique((const FUniqueNetIdSteam&)PlayerId);
+			// Should update MuteList after going up to the server and coming back down
 			ProcessMuteChangeNotification();
 		}
 		else
@@ -472,6 +479,7 @@ bool FOnlineVoiceSteam::UnmuteRemoteTalker(uint8 LocalUserNum, const FUniqueNetI
 		{
 			// Remove them from the mute list
 			SystemMuteList.RemoveSingleSwap((const FUniqueNetIdSteam&)PlayerId);
+			// Should update MuteList after going up to the server and coming back down
 			ProcessMuteChangeNotification();
 		}
 		else
@@ -486,8 +494,8 @@ bool FOnlineVoiceSteam::UnmuteRemoteTalker(uint8 LocalUserNum, const FUniqueNetI
 				if (Talker != NULL)
 				{
 					// Make sure there isn't a system mute
-					bool bIsMuted = IsMuted(LocalUserNum, PlayerId);
-					if (!bIsMuted)
+					bool bIsSystemMuted = IsSystemWideMuted(PlayerId);
+					if (!bIsSystemMuted)
 					{
 						// Remove them from the mute list
 						MuteList.RemoveSingleSwap((const FUniqueNetIdSteam&)PlayerId);
diff --git a/Engine/Plugins/Online/OnlineSubsystemSteam/Source/Private/VoiceInterfaceSteam.h b/Engine/Plugins/Online/OnlineSubsystemSteam/Source/Private/VoiceInterfaceSteam.h
--- a/Engine/Plugins/Online/OnlineSubsystemSteam/Source/Private/VoiceInterfaceSteam.h
+++ b/Engine/Plugins/Online/OnlineSubsystemSteam/Source/Private/VoiceInterfaceSteam.h
@@ -34,7 +34,7 @@ class FOnlineVoiceSteam : public IOnlineVoice
 	TArray<FRemoteTalker> RemoteTalkers;
 	/** Remote players locally muted explicitly */
 	TArray<FUniqueNetIdSteam> SystemMuteList;
-	/** Remote players locally muted */
+	/** Remote players locally muted (super set of SystemMuteList) */
 	TArray<FUniqueNetIdSteam> MuteList;
 
 	/** Time to wait for new data before triggering "not talking" */
@@ -54,9 +54,22 @@ class FOnlineVoiceSteam : public IOnlineVoice
 
 	/**
 	 * Is a given id presently muted (either by system mute or game server)
+	 *
+	 * @param UniqueId the net id to query
+	 *
+	 * @return true if the net id is muted at all, false otherwise
 	 */
 	bool IsLocallyMuted(const FUniqueNetId& UniqueId) const;
 
+	/**
+	 * Does a given id exist in the system wide mute list
+	 *
+	 * @param UniqueId the net id to query
+	 *
+	 * @return true if the net id is on the system wide mute list, false otherwise
+	 */
+	bool IsSystemWideMuted(const FUniqueNetId& UniqueId) const;
+
 PACKAGE_SCOPE:
 	FOnlineVoiceSteam() :
 		SteamSubsystem(NULL),
diff --git a/Engine/Plugins/Online/OnlineSubsystemUtils/Source/OnlineBlueprintSupport/OnlineBlueprintSupport.Build.cs b/Engine/Plugins/Online/OnlineSubsystemUtils/Source/OnlineBlueprintSupport/OnlineBlueprintSupport.Build.cs
--- a/Engine/Plugins/Online/OnlineSubsystemUtils/Source/OnlineBlueprintSupport/OnlineBlueprintSupport.Build.cs
+++ b/Engine/Plugins/Online/OnlineSubsystemUtils/Source/OnlineBlueprintSupport/OnlineBlueprintSupport.Build.cs
@@ -6,6 +6,8 @@ public class OnlineBlueprintSupport : ModuleRules
 {
 	public OnlineBlueprintSupport(TargetInfo Target)
 	{
+		PCHUsage = ModuleRules.PCHUsageMode.UseExplicitOrSharedPCHs;
+		
         PrivateDependencyModuleNames.AddRange(
 			new string[] { 
 				"Core", 
diff --git a/Engine/Plugins/Online/OnlineSubsystemUtils/Source/OnlineSubsystemUtils/Classes/OnlineSessionClient.h b/Engine/Plugins/Online/OnlineSubsystemUtils/Source/OnlineSubsystemUtils/Classes/OnlineSessionClient.h
--- a/Engine/Plugins/Online/OnlineSubsystemUtils/Source/OnlineSubsystemUtils/Classes/OnlineSessionClient.h
+++ b/Engine/Plugins/Online/OnlineSubsystemUtils/Source/OnlineSubsystemUtils/Classes/OnlineSessionClient.h
@@ -5,6 +5,7 @@
  */
 
 #pragma once
+#include "OnlineSubsystem.h"
 
 #include "CoreMinimal.h"
 #include "UObject/ObjectMacros.h"
diff --git a/Engine/Plugins/Online/OnlineSubsystemUtils/Source/OnlineSubsystemUtils/OnlineSubsystemUtils.Build.cs b/Engine/Plugins/Online/OnlineSubsystemUtils/Source/OnlineSubsystemUtils/OnlineSubsystemUtils.Build.cs
--- a/Engine/Plugins/Online/OnlineSubsystemUtils/Source/OnlineSubsystemUtils/OnlineSubsystemUtils.Build.cs
+++ b/Engine/Plugins/Online/OnlineSubsystemUtils/Source/OnlineSubsystemUtils/OnlineSubsystemUtils.Build.cs
@@ -7,10 +7,11 @@ public class OnlineSubsystemUtils : ModuleRules
 	public OnlineSubsystemUtils(TargetInfo Target)
 	{
 		Definitions.Add("ONLINESUBSYSTEMUTILS_PACKAGE=1");
+		PCHUsage = ModuleRules.PCHUsageMode.UseExplicitOrSharedPCHs;
 
         PrivateIncludePaths.Add("OnlineSubsystemUtils/Private");
 
-        PrivateDependencyModuleNames.AddRange(
+		PrivateDependencyModuleNames.AddRange(
 			new string[] { 
 				"Core", 
 				"CoreUObject",
diff --git a/Engine/Plugins/Online/OnlineSubsystemUtils/Source/OnlineSubsystemUtils/Private/OnlinePIESettings.cpp b/Engine/Plugins/Online/OnlineSubsystemUtils/Source/OnlineSubsystemUtils/Private/OnlinePIESettings.cpp
--- a/Engine/Plugins/Online/OnlineSubsystemUtils/Source/OnlineSubsystemUtils/Private/OnlinePIESettings.cpp
+++ b/Engine/Plugins/Online/OnlineSubsystemUtils/Source/OnlineSubsystemUtils/Private/OnlinePIESettings.cpp
@@ -110,14 +110,14 @@ void UOnlinePIESettings::PostEditChangeProperty(FPropertyChangedEvent& PropertyC
 	{
 		FName MemberPropName = PropertyChangedEvent.MemberProperty->GetFName();
 		FName SubPropName = PropertyChangedEvent.Property->GetFName();
-		if (MemberPropName == "bOnlinePIEEnabled")
+		if (MemberPropName == GET_MEMBER_NAME_CHECKED(UOnlinePIESettings, bOnlinePIEEnabled))
 		{
 			// Possibly get rid of the null subsystem in favor of the real default or if we are disabling online pie then get rid of the real subsystem to replace it with null
 			IOnlineSubsystem::ReloadDefaultSubsystem();
 		}
-		if (MemberPropName == "Logins")
+		if (MemberPropName == GET_MEMBER_NAME_CHECKED(UOnlinePIESettings, Logins))
 		{
-			if (SubPropName == "Id")
+			if (SubPropName == GET_MEMBER_NAME_CHECKED(FPIELoginSettingsInternal, Id))
 			{
 				for (FPIELoginSettingsInternal& Login : Logins)
 				{
@@ -125,7 +125,7 @@ void UOnlinePIESettings::PostEditChangeProperty(FPropertyChangedEvent& PropertyC
 					Login.Id = Login.Id.Trim().TrimTrailing();
 				}
 			}
-			else if (SubPropName == "Token")
+			else if (SubPropName == GET_MEMBER_NAME_CHECKED(FPIELoginSettingsInternal, Token))
 			{
 				for (FPIELoginSettingsInternal& Login : Logins)
 				{
@@ -135,7 +135,7 @@ void UOnlinePIESettings::PostEditChangeProperty(FPropertyChangedEvent& PropertyC
 					Login.Encrypt();
 				}
 			}
-			else if (SubPropName == "Type")
+			else if (SubPropName == GET_MEMBER_NAME_CHECKED(FPIELoginSettingsInternal, Type))
 			{
 				for (FPIELoginSettingsInternal& Login : Logins)
 				{
diff --git a/Engine/Plugins/Online/OnlineSubsystemUtils/Source/OnlineSubsystemUtils/Private/VoiceInterfaceImpl.cpp b/Engine/Plugins/Online/OnlineSubsystemUtils/Source/OnlineSubsystemUtils/Private/VoiceInterfaceImpl.cpp
--- a/Engine/Plugins/Online/OnlineSubsystemUtils/Source/OnlineSubsystemUtils/Private/VoiceInterfaceImpl.cpp
+++ b/Engine/Plugins/Online/OnlineSubsystemUtils/Source/OnlineSubsystemUtils/Private/VoiceInterfaceImpl.cpp
@@ -438,13 +438,12 @@ bool FOnlineVoiceImpl::IsRemotePlayerTalking(const FUniqueNetId& UniqueId)
 
 bool FOnlineVoiceImpl::IsMuted(uint32 LocalUserNum, const FUniqueNetId& UniqueId) const
 {
-	int32 Index = INDEX_NONE;
 	if (LocalUserNum >= 0 && LocalUserNum < (uint32)MaxLocalTalkers)
 	{
-		Index = SystemMuteList.Find((const FUniqueNetIdString&)UniqueId);
+		return IsLocallyMuted(UniqueId);
 	}
 
-	return Index != INDEX_NONE;
+	return false;
 }
 
 bool FOnlineVoiceImpl::IsLocallyMuted(const FUniqueNetId& UniqueId) const
@@ -453,6 +452,12 @@ bool FOnlineVoiceImpl::IsLocallyMuted(const FUniqueNetId& UniqueId) const
 	return Index != INDEX_NONE;
 }
 
+bool FOnlineVoiceImpl::IsSystemWideMuted(const FUniqueNetId& UniqueId) const
+{
+	int32 Index = SystemMuteList.Find((const FUniqueNetIdString&)UniqueId);
+	return Index != INDEX_NONE;
+}
+
 bool FOnlineVoiceImpl::MuteRemoteTalker(uint8 LocalUserNum, const FUniqueNetId& PlayerId, bool bIsSystemWide)
 {
 	uint32 Return = E_FAIL;
@@ -460,7 +465,9 @@ bool FOnlineVoiceImpl::MuteRemoteTalker(uint8 LocalUserNum, const FUniqueNetId&
 	{
 		if (bIsSystemWide)
 		{
+			// Add them to the system wide mute list
 			SystemMuteList.AddUnique((const FUniqueNetIdString&)PlayerId);
+			// Should update MuteList after going up to the server and coming back down
 			ProcessMuteChangeNotification();
 		}
 		else
@@ -500,8 +507,9 @@ bool FOnlineVoiceImpl::UnmuteRemoteTalker(uint8 LocalUserNum, const FUniqueNetId
 	{
 		if (bIsSystemWide)
 		{
-			// Remove them from the mute list
+			// Remove them from the system wide mute list
 			SystemMuteList.RemoveSingleSwap((const FUniqueNetIdString&)PlayerId);
+			// Should update MuteList after going up to the server and coming back down
 			ProcessMuteChangeNotification();
 		}
 		else
@@ -516,8 +524,8 @@ bool FOnlineVoiceImpl::UnmuteRemoteTalker(uint8 LocalUserNum, const FUniqueNetId
 				if (Talker != NULL)
 				{
 					// Make sure there isn't a system mute
-					bool bIsMuted = IsMuted(LocalUserNum, PlayerId);
-					if (!bIsMuted)
+					bool bIsSystemMuted = IsSystemWideMuted(PlayerId);
+					if (!bIsSystemMuted)
 					{
 						// Remove them from the mute list
 						MuteList.RemoveSingleSwap((const FUniqueNetIdString&)PlayerId);
diff --git a/Engine/Plugins/Online/OnlineSubsystemUtils/Source/OnlineSubsystemUtils/Public/PartyBeaconClient.h b/Engine/Plugins/Online/OnlineSubsystemUtils/Source/OnlineSubsystemUtils/Public/PartyBeaconClient.h
--- a/Engine/Plugins/Online/OnlineSubsystemUtils/Source/OnlineSubsystemUtils/Public/PartyBeaconClient.h
+++ b/Engine/Plugins/Online/OnlineSubsystemUtils/Source/OnlineSubsystemUtils/Public/PartyBeaconClient.h
@@ -28,7 +28,9 @@ enum class EClientRequestType : uint8
 	/** Reservation to configure an empty server  */
 	EmptyServerReservation,
 	/** Simple reconnect (checks for existing reservation) */
-	Reconnect
+	Reconnect,
+	/** Abandon the reservation beacon (game specific handling)*/
+	Abandon
 };
 
 inline const TCHAR* ToString(EClientRequestType RequestType)
diff --git a/Engine/Plugins/Online/OnlineSubsystemUtils/Source/OnlineSubsystemUtils/Public/VoiceInterfaceImpl.h b/Engine/Plugins/Online/OnlineSubsystemUtils/Source/OnlineSubsystemUtils/Public/VoiceInterfaceImpl.h
--- a/Engine/Plugins/Online/OnlineSubsystemUtils/Source/OnlineSubsystemUtils/Public/VoiceInterfaceImpl.h
+++ b/Engine/Plugins/Online/OnlineSubsystemUtils/Source/OnlineSubsystemUtils/Public/VoiceInterfaceImpl.h
@@ -35,7 +35,7 @@ class ONLINESUBSYSTEMUTILS_API FOnlineVoiceImpl : public IOnlineVoice
 	TArray<FRemoteTalker> RemoteTalkers;
 	/** Remote players locally muted explicitly */
 	TArray<FUniqueNetIdString> SystemMuteList;
-	/** All remote players locally muted */
+	/** All remote players locally muted (super set of SystemMuteList) */
 	TArray<FUniqueNetIdString> MuteList;
 
 	/** Time to wait for new data before triggering "not talking" */
@@ -55,9 +55,22 @@ class ONLINESUBSYSTEMUTILS_API FOnlineVoiceImpl : public IOnlineVoice
 
 	/**
 	 * Is a given id presently muted (either by system mute or game server)
+	 *
+	 * @param UniqueId the net id to query
+	 *
+	 * @return true if the net id is muted at all, false otherwise
 	 */
 	bool IsLocallyMuted(const FUniqueNetId& UniqueId) const;
 
+	/**
+	 * Does a given id exist in the system wide mute list
+	 *
+	 * @param UniqueId the net id to query
+	 *
+	 * @return true if the net id is on the system wide mute list, false otherwise
+	 */
+	bool IsSystemWideMuted(const FUniqueNetId& UniqueId) const;
+
 PACKAGE_SCOPE:
 
 	FOnlineVoiceImpl() :
diff --git a/Engine/Source/Runtime/Launch/Private/Android/AndroidJNI.cpp b/Engine/Source/Runtime/Launch/Private/Android/AndroidJNI.cpp
--- a/Engine/Source/Runtime/Launch/Private/Android/AndroidJNI.cpp
+++ b/Engine/Source/Runtime/Launch/Private/Android/AndroidJNI.cpp
@@ -9,9 +9,13 @@
 #include "AndroidInputInterface.h"
 #include "Widgets/Input/IVirtualKeyboardEntry.h"
 #include "UnrealEngine.h"
+#include "Misc/ConfigCacheIni.h"
 #include "Misc/FeedbackContext.h"
+
+THIRD_PARTY_INCLUDES_START
 #include <android/asset_manager.h>
 #include <android/asset_manager_jni.h>
+THIRD_PARTY_INCLUDES_END
 
 #define JNI_CURRENT_VERSION JNI_VERSION_1_6
 
@@ -31,7 +35,7 @@ FOnActivityResult FJavaWrapper::OnActivityResultDelegate;
 //////////////////////////////////////////////////////////////////////////
 
 #if UE_BUILD_SHIPPING
-// always clear any exceptions in SHipping
+// always clear any exceptions in Shipping
 #define CHECK_JNI_RESULT(Id) if (Id == 0) { Env->ExceptionClear(); }
 #else
 #define CHECK_JNI_RESULT(Id) \
@@ -42,6 +46,8 @@ if (Id == 0) \
 }
 #endif
 
+#define CHECK_JNI_METHOD(Id) checkf(Id != nullptr, TEXT("Failed to find " #Id));
+
 void FJavaWrapper::FindClassesAndMethods(JNIEnv* Env)
 {
 	bool bIsOptional = false;
@@ -89,6 +95,11 @@ void FJavaWrapper::FindClassesAndMethods(JNIEnv* Env)
 	InputDeviceInfo_Name = FJavaWrapper::FindField(Env, InputDeviceInfoClass, "name", "Ljava/lang/String;", bIsOptional);
 	InputDeviceInfo_Descriptor = FJavaWrapper::FindField(Env, InputDeviceInfoClass, "descriptor", "Ljava/lang/String;", bIsOptional);
 
+	/** GooglePlay services */
+	FindGooglePlayMethods(Env);
+	/** GooglePlay billing services */
+	FindGooglePlayBillingMethods(Env);
+
 	// get field IDs for LaunchNotificationClass class members
 	jclass localLaunchNotificationClass = FindClass(Env, "com/epicgames/ue4/GameActivity$LaunchNotification", bIsOptional);
 	LaunchNotificationClass = (jclass)Env->NewGlobalRef(localLaunchNotificationClass);
@@ -99,6 +110,16 @@ void FJavaWrapper::FindClassesAndMethods(JNIEnv* Env)
 
 	// the rest are optional
 	bIsOptional = true;
+
+	// SurfaceView functionality for view scaling on some devices
+	AndroidThunkJava_UseSurfaceViewWorkaround = FindMethod(Env, GameActivityClassID, "AndroidThunkJava_UseSurfaceViewWorkaround", "()V", bIsOptional);
+	AndroidThunkJava_SetDesiredViewSize = FindMethod(Env, GameActivityClassID, "AndroidThunkJava_SetDesiredViewSize", "(II)V", bIsOptional);
+}
+
+void FJavaWrapper::FindGooglePlayMethods(JNIEnv* Env)
+{
+	bool bIsOptional = true;
+
 	// @todo split GooglePlay
 	//	GoogleServicesClassID = FindClass(Env, "com/epicgames/ue4/GoogleServices", bIsOptional);
 	GoogleServicesClassID = GameActivityClassID;
@@ -112,20 +133,27 @@ void FJavaWrapper::FindClassesAndMethods(JNIEnv* Env)
 	AndroidThunkJava_ShowInterstitialAd = FindMethod(Env, GoogleServicesClassID, "AndroidThunkJava_ShowInterstitialAd", "()V", bIsOptional);
 	AndroidThunkJava_GoogleClientConnect = FindMethod(Env, GoogleServicesClassID, "AndroidThunkJava_GoogleClientConnect", "()V", bIsOptional);
 	AndroidThunkJava_GoogleClientDisconnect = FindMethod(Env, GoogleServicesClassID, "AndroidThunkJava_GoogleClientDisconnect", "()V", bIsOptional);
-
+}
+void FJavaWrapper::FindGooglePlayBillingMethods(JNIEnv* Env)
+{
 	// In app purchase functionality
+	bool bSupportsInAppPurchasing = false;
+	if (!GConfig->GetBool(TEXT("OnlineSubsystemGooglePlay.Store"), TEXT("bSupportsInAppPurchasing"), bSupportsInAppPurchasing, GEngineIni))
+	{
+		FPlatformMisc::LowLevelOutputDebugString(TEXT("[JNI] - Failed to determine if app purchasing is enabled!"));
+	}
+	bool bIsStoreOptional = !bSupportsInAppPurchasing;
+
 	jclass localStringClass = Env->FindClass("java/lang/String");
 	JavaStringClass = (jclass)Env->NewGlobalRef(localStringClass);
 	Env->DeleteLocalRef(localStringClass);
-	AndroidThunkJava_IapSetupService = FindMethod(Env, GoogleServicesClassID, "AndroidThunkJava_IapSetupService", "(Ljava/lang/String;)V", bIsOptional);
-	AndroidThunkJava_IapQueryInAppPurchases = FindMethod(Env, GoogleServicesClassID, "AndroidThunkJava_IapQueryInAppPurchases", "([Ljava/lang/String;[Z)Z", bIsOptional);
-	AndroidThunkJava_IapBeginPurchase = FindMethod(Env, GoogleServicesClassID, "AndroidThunkJava_IapBeginPurchase", "(Ljava/lang/String;Z)Z", bIsOptional);
-	AndroidThunkJava_IapIsAllowedToMakePurchases = FindMethod(Env, GoogleServicesClassID, "AndroidThunkJava_IapIsAllowedToMakePurchases", "()Z", bIsOptional);
-	AndroidThunkJava_IapRestorePurchases = FindMethod(Env, GoogleServicesClassID, "AndroidThunkJava_IapRestorePurchases", "([Ljava/lang/String;[Z)Z", bIsOptional);
-
-	// SurfaceView functionality for view scaling on some devices
-	AndroidThunkJava_UseSurfaceViewWorkaround = FindMethod(Env, GameActivityClassID, "AndroidThunkJava_UseSurfaceViewWorkaround", "()V", bIsOptional);
-	AndroidThunkJava_SetDesiredViewSize = FindMethod(Env, GameActivityClassID, "AndroidThunkJava_SetDesiredViewSize", "(II)V", bIsOptional);
+	AndroidThunkJava_IapSetupService = FindMethod(Env, GoogleServicesClassID, "AndroidThunkJava_IapSetupService", "(Ljava/lang/String;)V", bIsStoreOptional);
+	AndroidThunkJava_IapQueryInAppPurchases = FindMethod(Env, GoogleServicesClassID, "AndroidThunkJava_IapQueryInAppPurchases", "([Ljava/lang/String;)Z", bIsStoreOptional);
+	AndroidThunkJava_IapBeginPurchase = FindMethod(Env, GoogleServicesClassID, "AndroidThunkJava_IapBeginPurchase", "(Ljava/lang/String;)Z", bIsStoreOptional);
+	AndroidThunkJava_IapIsAllowedToMakePurchases = FindMethod(Env, GoogleServicesClassID, "AndroidThunkJava_IapIsAllowedToMakePurchases", "()Z", bIsStoreOptional);
+	AndroidThunkJava_IapRestorePurchases = FindMethod(Env, GoogleServicesClassID, "AndroidThunkJava_IapRestorePurchases", "([Ljava/lang/String;[Z)Z", bIsStoreOptional);
+	AndroidThunkJava_IapConsumePurchase = FindMethod(Env, GoogleServicesClassID, "AndroidThunkJava_IapConsumePurchase", "(Ljava/lang/String;)Z", bIsStoreOptional);
+	AndroidThunkJava_IapQueryExistingPurchases = FindMethod(Env, GoogleServicesClassID, "AndroidThunkJava_IapQueryExistingPurchases", "()Z", bIsStoreOptional);
 }
 
 jclass FJavaWrapper::FindClass(JNIEnv* Env, const ANSICHAR* ClassName, bool bIsOptional)
@@ -273,6 +301,8 @@ jmethodID FJavaWrapper::AndroidThunkJava_IapQueryInAppPurchases;
 jmethodID FJavaWrapper::AndroidThunkJava_IapBeginPurchase;
 jmethodID FJavaWrapper::AndroidThunkJava_IapIsAllowedToMakePurchases;
 jmethodID FJavaWrapper::AndroidThunkJava_IapRestorePurchases;
+jmethodID FJavaWrapper::AndroidThunkJava_IapQueryExistingPurchases;
+jmethodID FJavaWrapper::AndroidThunkJava_IapConsumePurchase;
 
 jmethodID FJavaWrapper::AndroidThunkJava_UseSurfaceViewWorkaround;
 jmethodID FJavaWrapper::AndroidThunkJava_SetDesiredViewSize;
@@ -800,59 +830,109 @@ void AndroidThunkCpp_Iap_SetupIapService(const FString& InProductKey)
 	}
 }
 
-bool AndroidThunkCpp_Iap_QueryInAppPurchases(const TArray<FString>& ProductIDs, const TArray<bool>& bConsumable)
+bool AndroidThunkCpp_Iap_QueryInAppPurchases(const TArray<FString>& ProductIDs)
 {
-	FPlatformMisc::LowLevelOutputDebugString(L"[JNI] - AndroidThunkCpp_Iap_QueryInAppPurchases");
+	FPlatformMisc::LowLevelOutputDebugString(TEXT("[JNI] - AndroidThunkCpp_Iap_QueryInAppPurchases"));
 	bool bResult = false;
 
 	if (JNIEnv* Env = FAndroidApplication::GetJavaEnv())
 	{
+		CHECK_JNI_METHOD(FJavaWrapper::AndroidThunkJava_IapQueryInAppPurchases);
+
 		// Populate some java types with the provided product information
 		jobjectArray ProductIDArray = (jobjectArray)Env->NewObjectArray(ProductIDs.Num(), FJavaWrapper::JavaStringClass, NULL);
-		jbooleanArray ConsumeArray = (jbooleanArray)Env->NewBooleanArray(ProductIDs.Num());
-
-		jboolean* ConsumeArrayValues = Env->GetBooleanArrayElements(ConsumeArray, 0);
 		for (uint32 Param = 0; Param < ProductIDs.Num(); Param++)
 		{
 			jstring StringValue = Env->NewStringUTF(TCHAR_TO_UTF8(*ProductIDs[Param]));
 			Env->SetObjectArrayElement(ProductIDArray, Param, StringValue);
 			Env->DeleteLocalRef(StringValue);
-
-			ConsumeArrayValues[Param] = bConsumable[Param];
 		}
-		Env->ReleaseBooleanArrayElements(ConsumeArray, ConsumeArrayValues, 0);
 
 		// Execute the java code for this operation
-		bResult = FJavaWrapper::CallBooleanMethod(Env, FJavaWrapper::GoogleServicesThis, FJavaWrapper::AndroidThunkJava_IapQueryInAppPurchases, ProductIDArray, ConsumeArray);
+		bResult = FJavaWrapper::CallBooleanMethod(Env, FJavaWrapper::GoogleServicesThis, FJavaWrapper::AndroidThunkJava_IapQueryInAppPurchases, ProductIDArray);
 		
 		// clean up references
 		Env->DeleteLocalRef(ProductIDArray);
-		Env->DeleteLocalRef(ConsumeArray);
 	}
 
 	return bResult;
 }
 
-bool AndroidThunkCpp_Iap_BeginPurchase(const FString& ProductID, const bool bConsumable)
+bool AndroidThunkCpp_Iap_QueryInAppPurchases(const TArray<FString>& ProductIDs, const TArray<bool>& bConsumable)
 {
-	FPlatformMisc::LowLevelOutputDebugString(L"[JNI] - AndroidThunkCpp_Iap_BeginPurchase");
+	FPlatformMisc::LowLevelOutputDebugString(TEXT("AndroidThunkCpp_Iap_QueryInAppPurchases DEPRECATED, won't use consumables array"));
+	return AndroidThunkCpp_Iap_QueryInAppPurchases(ProductIDs);
+}
+
+bool AndroidThunkCpp_Iap_BeginPurchase(const FString& ProductID)
+{
+	FPlatformMisc::LowLevelOutputDebugStringf(TEXT("[JNI] - AndroidThunkCpp_Iap_BeginPurchase %s"), *ProductID);
 	bool bResult = false;
 	if (JNIEnv* Env = FAndroidApplication::GetJavaEnv())
 	{
+		CHECK_JNI_METHOD(FJavaWrapper::AndroidThunkJava_IapBeginPurchase);
+
 		jstring ProductIdJava = Env->NewStringUTF(TCHAR_TO_UTF8(*ProductID));
-		bResult = FJavaWrapper::CallBooleanMethod(Env, FJavaWrapper::GoogleServicesThis, FJavaWrapper::AndroidThunkJava_IapBeginPurchase, ProductIdJava, bConsumable);
+		bResult = FJavaWrapper::CallBooleanMethod(Env, FJavaWrapper::GoogleServicesThis, FJavaWrapper::AndroidThunkJava_IapBeginPurchase, ProductIdJava);
 		Env->DeleteLocalRef(ProductIdJava);
 	}
 
 	return bResult;
 }
 
+bool AndroidThunkCpp_Iap_BeginPurchase(const FString& ProductID, const bool bConsumable)
+{
+	FPlatformMisc::LowLevelOutputDebugString(TEXT("AndroidThunkCpp_Iap_BeginPurchase DEPRECATED, won't use consumable flag"));
+	return AndroidThunkCpp_Iap_BeginPurchase(ProductID);
+}
+
+bool AndroidThunkCpp_Iap_ConsumePurchase(const FString& ProductToken)
+{
+	FPlatformMisc::LowLevelOutputDebugStringf(TEXT("[JNI] - AndroidThunkCpp_Iap_ConsumePurchase %s"), *ProductToken);
+	
+	bool bResult = false;
+	if (!ProductToken.IsEmpty())
+	{
+		if (JNIEnv* Env = FAndroidApplication::GetJavaEnv())
+		{
+			CHECK_JNI_METHOD(FJavaWrapper::AndroidThunkJava_IapConsumePurchase);
+
+			jstring ProductTokenJava = Env->NewStringUTF(TCHAR_TO_UTF8(*ProductToken));
+			//FPlatformMisc::LowLevelOutputDebugStringf(TEXT("[JNI] - AndroidThunkCpp_Iap_ConsumePurchase BEGIN"));
+			bResult = FJavaWrapper::CallBooleanMethod(Env, FJavaWrapper::GoogleServicesThis, FJavaWrapper::AndroidThunkJava_IapConsumePurchase, ProductTokenJava);
+			//FPlatformMisc::LowLevelOutputDebugStringf(TEXT("[JNI] - AndroidThunkCpp_Iap_ConsumePurchase END"));
+			Env->DeleteLocalRef(ProductTokenJava);
+		}
+	}
+
+	return bResult;
+}
+
+bool AndroidThunkCpp_Iap_QueryExistingPurchases()
+{
+	FPlatformMisc::LowLevelOutputDebugStringf(TEXT("[JNI] - AndroidThunkCpp_Iap_QueryExistingPurchases"));
+	
+	bool bResult = false;
+	if (JNIEnv* Env = FAndroidApplication::GetJavaEnv())
+	{
+		CHECK_JNI_METHOD(FJavaWrapper::AndroidThunkJava_IapQueryExistingPurchases);
+
+		//FPlatformMisc::LowLevelOutputDebugStringf(TEXT("[JNI] - AndroidThunkCpp_Iap_QueryExistingPurchases BEGIN"));
+		bResult = FJavaWrapper::CallBooleanMethod(Env, FJavaWrapper::GoogleServicesThis, FJavaWrapper::AndroidThunkJava_IapQueryExistingPurchases);
+		//FPlatformMisc::LowLevelOutputDebugStringf(TEXT("[JNI] - AndroidThunkCpp_Iap_QueryExistingPurchases END"));
+	}
+
+	return bResult;
+}
+
 bool AndroidThunkCpp_Iap_IsAllowedToMakePurchases()
 {
-	FPlatformMisc::LowLevelOutputDebugString(L"[JNI] - AndroidThunkCpp_Iap_IsAllowedToMakePurchases");
+	FPlatformMisc::LowLevelOutputDebugString(TEXT("[JNI] - AndroidThunkCpp_Iap_IsAllowedToMakePurchases"));
 	bool bResult = false;
 	if (JNIEnv* Env = FAndroidApplication::GetJavaEnv())
 	{
+		CHECK_JNI_METHOD(FJavaWrapper::AndroidThunkJava_IapIsAllowedToMakePurchases);
+
 		bResult = FJavaWrapper::CallBooleanMethod(Env, FJavaWrapper::GoogleServicesThis, FJavaWrapper::AndroidThunkJava_IapIsAllowedToMakePurchases);
 	}
 	return bResult;
@@ -860,11 +940,13 @@ bool AndroidThunkCpp_Iap_IsAllowedToMakePurchases()
 
 bool AndroidThunkCpp_Iap_RestorePurchases(const TArray<FString>& ProductIDs, const TArray<bool>& bConsumable)
 {
-	FPlatformMisc::LowLevelOutputDebugString(L"[JNI] - AndroidThunkCpp_Iap_RestorePurchases");
+	FPlatformMisc::LowLevelOutputDebugString(TEXT("[JNI] - AndroidThunkCpp_Iap_RestorePurchases"));
 	bool bResult = false;
 
 	if (JNIEnv* Env = FAndroidApplication::GetJavaEnv())
 	{
+		CHECK_JNI_METHOD(FJavaWrapper::AndroidThunkJava_IapRestorePurchases);
+
 		// Populate some java types with the provided product information
 		jobjectArray ProductIDArray = (jobjectArray)Env->NewObjectArray(ProductIDs.Num(), FJavaWrapper::JavaStringClass, NULL);
 		jbooleanArray ConsumeArray = (jbooleanArray)Env->NewBooleanArray(ProductIDs.Num());
@@ -1001,12 +1083,12 @@ bool AndroidThunkCpp_HasActiveWiFiConnection()
 
 JNIEXPORT jint JNI_OnLoad(JavaVM* InJavaVM, void* InReserved)
 {
-	FPlatformMisc::LowLevelOutputDebugString(L"In the JNI_OnLoad function");
+	FPlatformMisc::LowLevelOutputDebugString(TEXT("In the JNI_OnLoad function"));
 
 	JNIEnv* Env = NULL;
 	InJavaVM->GetEnv((void **)&Env, JNI_CURRENT_VERSION);
 
-	// if you have problems with stuff being missing esspecially in distribution builds then it could be because proguard is stripping things from java
+	// if you have problems with stuff being missing especially in distribution builds then it could be because proguard is stripping things from java
 	// check proguard-project.txt and see if your stuff is included in the exceptions
 	GJavaVM = InJavaVM;
 	FAndroidApplication::InitializeJavaEnv(GJavaVM, JNI_CURRENT_VERSION, FJavaWrapper::GameActivityThis);
@@ -1051,8 +1133,7 @@ JNIEXPORT jint JNI_OnLoad(JavaVM* InJavaVM, void* InReserved)
 	extern CORE_API FAndroidLaunchURLDelegate OnAndroidLaunchURL;
 	OnAndroidLaunchURL = FAndroidLaunchURLDelegate::CreateStatic(&AndroidThunkCpp_LaunchURL);
 
-	FPlatformMisc::LowLevelOutputDebugString(L"In the JNI_OnLoad function 5");
-
+	FPlatformMisc::LowLevelOutputDebugString(TEXT("In the JNI_OnLoad function 5"));
 
 	return JNI_CURRENT_VERSION;
 }
@@ -1067,7 +1148,7 @@ extern "C" void Java_com_epicgames_ue4_GameActivity_nativeSetGlobalActivity(JNIE
 		FJavaWrapper::GameActivityThis = jenv->NewGlobalRef(thiz);
 		if (!FJavaWrapper::GameActivityThis)
 		{
-			FPlatformMisc::LowLevelOutputDebugString(L"Error setting the global GameActivity activity");
+			FPlatformMisc::LowLevelOutputDebugString(TEXT("Error setting the global GameActivity activity"));
 			check(false);
 		}
 
diff --git a/Engine/Source/Runtime/Launch/Public/Android/AndroidJNI.h b/Engine/Source/Runtime/Launch/Public/Android/AndroidJNI.h
--- a/Engine/Source/Runtime/Launch/Public/Android/AndroidJNI.h
+++ b/Engine/Source/Runtime/Launch/Public/Android/AndroidJNI.h
@@ -76,6 +76,8 @@ public:
 	static jmethodID AndroidThunkJava_IapBeginPurchase;
 	static jmethodID AndroidThunkJava_IapIsAllowedToMakePurchases;
 	static jmethodID AndroidThunkJava_IapRestorePurchases;
+	static jmethodID AndroidThunkJava_IapQueryExistingPurchases;
+	static jmethodID AndroidThunkJava_IapConsumePurchase;
 
 	// SurfaceView functionality for view scaling on some devices
 	static jmethodID AndroidThunkJava_UseSurfaceViewWorkaround;
@@ -107,4 +109,11 @@ public:
 
 	// Delegate that can be registered to that is called when an activity is finished
 	static FOnActivityResult OnActivityResultDelegate;
+
+private:
+
+	/** Find GooglePlay "game services" classes and methods */
+	static void FindGooglePlayMethods(JNIEnv* Env);
+	/** Find GooglePlay billing classes and methods */
+	static void FindGooglePlayBillingMethods(JNIEnv* Env);
 };
