#UE4 Exclude automationcontroller include for shipping since the module is not compiled in non-shipping and the include now depends on a generated file.

diff --git a/Engine/Source/Runtime/Launch/Private/LaunchEngineLoop.cpp b/Engine/Source/Runtime/Launch/Private/LaunchEngineLoop.cpp
--- a/Engine/Source/Runtime/Launch/Private/LaunchEngineLoop.cpp
+++ b/Engine/Source/Runtime/Launch/Private/LaunchEngineLoop.cpp
@@ -91,7 +91,9 @@
 	#include "EngineStats.h"
 	#include "EngineGlobals.h"
 	#include "AudioThread.h"
+#if WITH_ENGINE && !UE_BUILD_SHIPPING
 	#include "Interfaces/IAutomationControllerModule.h"
+#endif // WITH_ENGINE && !UE_BUILD_SHIPPING
 	#include "Database.h"
 	#include "DerivedDataCacheInterface.h"
 	#include "ShaderCompiler.h"
@@ -3114,7 +3116,7 @@ void FEngineLoop::Tick()
 		ClearPendingStatGroups();
 #endif
 
-#if WITH_EDITOR
+#if WITH_EDITOR && !UE_BUILD_SHIPPING
 		{
 			QUICK_SCOPE_CYCLE_COUNTER( STAT_FEngineLoop_Tick_AutomationController );
 			static FName AutomationController( "AutomationController" );
