Update branch name for analytics.

diff --git a/Engine/Source/Runtime/Launch/Resources/Version.h b/Engine/Source/Runtime/Launch/Resources/Version.h
--- a/Engine/Source/Runtime/Launch/Resources/Version.h
+++ b/Engine/Source/Runtime/Launch/Resources/Version.h
@@ -103,7 +103,7 @@
 
 // The branch that this engine is being built from. When set by UAT, this has the form of a Perforce depot path with forward slashes escaped by plus characters 
 // (eg. //UE4/Main -> ++UE4+Main)
-#define BRANCH_NAME "UE4"
+#define BRANCH_NAME "++UE4+Release-4.16"
 
 // Macros for encoding strings
 #define VERSION_TEXT_2(x) L ## x
