Copying //UE4/Dev-Sequencer to //UE4/Main (Source: //UE4/Dev-Sequencer @ 3617642)

diff --git a/Engine/Source/Runtime/Engine/Private/Texture.cpp b/Engine/Source/Runtime/Engine/Private/Texture.cpp
--- a/Engine/Source/Runtime/Engine/Private/Texture.cpp
+++ b/Engine/Source/Runtime/Engine/Private/Texture.cpp
@@ -1,9 +1,5 @@
 // Copyright 1998-2017 Epic Games, Inc. All Rights Reserved.
 
-/*=============================================================================
-	Texture.cpp: Implementation of UTexture.
-=============================================================================*/
-
 #include "Engine/Texture.h"
 #include "Misc/App.h"
 #include "Modules/ModuleManager.h"
@@ -18,16 +14,19 @@
 #include "EngineUtils.h"
 #include "Engine/AssetUserData.h"
 #include "EditorSupportDelegates.h"
-#include "Interfaces/IImageWrapperModule.h"
+#include "IImageWrapper.h"
+#include "IImageWrapperModule.h"
 #include "EngineGlobals.h"
 #include "Engine/Engine.h"
 #include "Interfaces/ITargetPlatform.h"
 
 #if WITH_EDITORONLY_DATA
-#include "EditorFramework/AssetImportData.h"
+	#include "EditorFramework/AssetImportData.h"
 #endif
+
 #include "Engine/TextureCube.h"
 
+
 DEFINE_LOG_CATEGORY(LogTexture);
 
 #if STATS
@@ -38,6 +37,7 @@ DECLARE_STATS_GROUP(TEXT("Texture Group"), STATGROUP_TextureGroup, STATCAT_Advan
 FOREACH_ENUM_TEXTUREGROUP(DECLARETEXTUREGROUPSTAT)
 #undef DECLARETEXTUREGROUPSTAT
 
+
 // Initialize TextureGroupStatFNames array with the FNames for each stats.
 FName FTextureResource::TextureGroupStatFNames[TEXTUREGROUP_MAX] =
 	{
@@ -132,16 +132,20 @@ void UTexture::PostEditChangeProperty(FPropertyChangedEvent& PropertyChangedEven
 	UProperty* PropertyThatChanged = PropertyChangedEvent.Property;
 	if( PropertyThatChanged )
 	{
-		static const FName CompressionSettingsName("CompressionSettings");
-		static const FName LODGroupName("LODGroup");
-		static const FName DeferCompressionName("DeferCompression");
+		static const FName CompressionSettingsName = GET_MEMBER_NAME_CHECKED(UTexture, CompressionSettings);
+		static const FName LODGroupName = GET_MEMBER_NAME_CHECKED(UTexture, LODGroup);
+		static const FName DeferCompressionName = GET_MEMBER_NAME_CHECKED(UTexture, DeferCompression);
+		static const FName SrgbName = GET_MEMBER_NAME_CHECKED(UTexture, SRGB);
 #if WITH_EDITORONLY_DATA
-		static const FName MaxTextureSizeName("MaxTextureSize");
-		static const FName CompressionQualityName("CompressionQuality");
-#endif // #if WITH_EDITORONLY_DATA
+		static const FName MaxTextureSizeName = GET_MEMBER_NAME_CHECKED(UTexture, MaxTextureSize);
+		static const FName CompressionQualityName = GET_MEMBER_NAME_CHECKED(UTexture, CompressionQuality);
+#endif //WITH_EDITORONLY_DATA
 
 		const FName PropertyName = PropertyThatChanged->GetFName();
-		if (PropertyName == CompressionSettingsName || PropertyName == LODGroupName)
+
+		if ((PropertyName == CompressionSettingsName) ||
+			(PropertyName == LODGroupName) ||
+			(PropertyName == SrgbName))
 		{
 			RequiresNotifyMaterials = true;
 		}
@@ -165,7 +169,7 @@ void UTexture::PostEditChangeProperty(FPropertyChangedEvent& PropertyChangedEven
 				MaxTextureSize = FMath::Min<int32>(FMath::RoundUpToPowerOfTwo(MaxTextureSize), GetMaximumDimension());
 			}
 		}
-#endif // #if WITH_EDITORONLY_DATA
+#endif //WITH_EDITORONLY_DATA
 
 		bool bPreventSRGB = (CompressionSettings == TC_Alpha || CompressionSettings == TC_Normalmap || CompressionSettings == TC_Masks || CompressionSettings == TC_HDR || CompressionSettings == TC_HDR_Compressed);
 		if(bPreventSRGB && SRGB == true)
@@ -212,43 +216,7 @@ void UTexture::PostEditChangeProperty(FPropertyChangedEvent& PropertyChangedEven
 	// Notify any loaded material instances if changed our compression format
 	if (RequiresNotifyMaterials)
 	{
-		TArray<UMaterialInterface*> MaterialsThatUseThisTexture;
-
-		// Create a material update context to safely update materials.
-		{
-			FMaterialUpdateContext UpdateContext;
-
-			// Notify any material that uses this texture
-			TSet<UMaterial*> BaseMaterialsThatUseThisTexture;
-			for (TObjectIterator<UMaterialInterface> It; It; ++It)
-			{
-				UMaterialInterface* MaterialInterface = *It;
-				if (DoesMaterialUseTexture(MaterialInterface,this))
-				{
-					MaterialsThatUseThisTexture.Add(MaterialInterface);
-
-					// This is a bit tricky. We want to make sure all materials using this texture are
-					// updated. Materials are always updated. Material instances may also have to be
-					// updated and if they have static permutations their children must be updated
-					// whether they use the texture or not! The safe thing to do is to add the instance's
-					// base material to the update context causing all materials in the tree to update.
-					BaseMaterialsThatUseThisTexture.Add(MaterialInterface->GetMaterial());
-				}
-			}
-
-			// Go ahead and update any base materials that need to be.
-			for (TSet<UMaterial*>::TConstIterator It(BaseMaterialsThatUseThisTexture); It; ++It)
-			{
-				UpdateContext.AddMaterial(*It);
-				(*It)->PostEditChange();
-			}
-		}
-
-		// Now that all materials and instances have updated send necessary callbacks.
-		for (int32 i = 0; i < MaterialsThatUseThisTexture.Num(); ++i)
-		{
-			FEditorSupportDelegates::MaterialTextureSettingsChanged.Broadcast(MaterialsThatUseThisTexture[i]);
-		}
+		NotifyMaterials();
 	}
 		
 #if WITH_EDITORONLY_DATA
@@ -668,9 +636,9 @@ void FTextureSource::Compress()
 	{
 		uint8* BulkDataPtr = (uint8*)BulkData.Lock(LOCK_READ_WRITE);
 		IImageWrapperModule& ImageWrapperModule = FModuleManager::LoadModuleChecked<IImageWrapperModule>( FName("ImageWrapper") );
-		IImageWrapperPtr ImageWrapper = ImageWrapperModule.CreateImageWrapper( EImageFormat::PNG );
+		TSharedPtr<IImageWrapper> ImageWrapper = ImageWrapperModule.CreateImageWrapper( EImageFormat::PNG );
 		// TODO: TSF_BGRA8 is stored as RGBA, so the R and B channels are swapped in the internal png. Should we fix this?
-		ERGBFormat::Type RawFormat = (Format == TSF_G8) ? ERGBFormat::Gray : ERGBFormat::RGBA;
+		ERGBFormat RawFormat = (Format == TSF_G8) ? ERGBFormat::Gray : ERGBFormat::RGBA;
 		if ( ImageWrapper.IsValid() && ImageWrapper->SetRaw( BulkDataPtr, BulkData.GetBulkDataSize(), SizeX, SizeY, RawFormat, Format == TSF_RGBA16 ? 16 : 8 ) )
 		{
 			const TArray<uint8>& CompressedData = ImageWrapper->GetCompressed();
@@ -710,14 +678,14 @@ uint8* FTextureSource::LockMip(int32 MipIndex)
 				}
 
 				IImageWrapperModule& ImageWrapperModule = FModuleManager::LoadModuleChecked<IImageWrapperModule>( FName("ImageWrapper") );
-				IImageWrapperPtr ImageWrapper = ImageWrapperModule.CreateImageWrapper( EImageFormat::PNG );
+				TSharedPtr<IImageWrapper> ImageWrapper = ImageWrapperModule.CreateImageWrapper( EImageFormat::PNG );
 				if ( ImageWrapper.IsValid() && ImageWrapper->SetCompressed( LockedMipData, BulkData.GetBulkDataSize() ) )
 				{
 					check( ImageWrapper->GetWidth() == SizeX );
 					check( ImageWrapper->GetHeight() == SizeY );
 					const TArray<uint8>* RawData = NULL;
 					// TODO: TSF_BGRA8 is stored as RGBA, so the R and B channels are swapped in the internal png. Should we fix this?
-					ERGBFormat::Type RawFormat = (Format == TSF_G8) ? ERGBFormat::Gray : ERGBFormat::RGBA;
+					ERGBFormat RawFormat = (Format == TSF_G8) ? ERGBFormat::Gray : ERGBFormat::RGBA;
 					if (ImageWrapper->GetRaw( RawFormat, Format == TSF_RGBA16 ? 16 : 8, RawData ))
 					{
 						if (RawData->Num() > 0)
@@ -787,7 +755,9 @@ bool FTextureSource::GetMipData(TArray<uint8>& OutMipData, int32 MipIndex, IImag
 				{
 					ImageWrapperModule = &FModuleManager::LoadModuleChecked<IImageWrapperModule>( FName("ImageWrapper") );
 				}
-				IImageWrapperPtr ImageWrapper = ImageWrapperModule->CreateImageWrapper( EImageFormat::PNG );
+
+				TSharedPtr<IImageWrapper> ImageWrapper = ImageWrapperModule->CreateImageWrapper( EImageFormat::PNG );
+
 				if ( ImageWrapper.IsValid() && ImageWrapper->SetCompressed( RawSourceData, BulkData.GetBulkDataSize() ) )
 				{
 					if (ImageWrapper->GetWidth() == SizeX
@@ -795,7 +765,7 @@ bool FTextureSource::GetMipData(TArray<uint8>& OutMipData, int32 MipIndex, IImag
 					{
 						const TArray<uint8>* RawData = NULL;
 						// TODO: TSF_BGRA8 is stored as RGBA, so the R and B channels are swapped in the internal png. Should we fix this?
-						ERGBFormat::Type RawFormat = (Format == TSF_G8) ? ERGBFormat::Gray : ERGBFormat::RGBA;
+						ERGBFormat RawFormat = (Format == TSF_G8) ? ERGBFormat::Gray : ERGBFormat::RGBA;
 						if (ImageWrapper->GetRaw( RawFormat, Format == TSF_RGBA16 ? 16 : 8, RawData ))
 						{
 							OutMipData = *RawData;
@@ -1163,3 +1133,47 @@ void GetAllDefaultTextureFormats(const class ITargetPlatform* TargetPlatform, TA
 #endif
 }
 
+#if WITH_EDITOR
+
+void UTexture::NotifyMaterials()
+{
+	TArray<UMaterialInterface*> MaterialsThatUseThisTexture;
+
+	// Create a material update context to safely update materials.
+	{
+		FMaterialUpdateContext UpdateContext;
+
+		// Notify any material that uses this texture
+		TSet<UMaterial*> BaseMaterialsThatUseThisTexture;
+		for (TObjectIterator<UMaterialInterface> It; It; ++It)
+		{
+			UMaterialInterface* MaterialInterface = *It;
+			if (DoesMaterialUseTexture(MaterialInterface, this))
+			{
+				MaterialsThatUseThisTexture.Add(MaterialInterface);
+
+				// This is a bit tricky. We want to make sure all materials using this texture are
+				// updated. Materials are always updated. Material instances may also have to be
+				// updated and if they have static permutations their children must be updated
+				// whether they use the texture or not! The safe thing to do is to add the instance's
+				// base material to the update context causing all materials in the tree to update.
+				BaseMaterialsThatUseThisTexture.Add(MaterialInterface->GetMaterial());
+			}
+		}
+
+		// Go ahead and update any base materials that need to be.
+		for (TSet<UMaterial*>::TConstIterator It(BaseMaterialsThatUseThisTexture); It; ++It)
+		{
+			UpdateContext.AddMaterial(*It);
+			(*It)->PostEditChange();
+		}
+	}
+
+	// Now that all materials and instances have updated send necessary callbacks.
+	for (int32 i = 0; i < MaterialsThatUseThisTexture.Num(); ++i)
+	{
+		FEditorSupportDelegates::MaterialTextureSettingsChanged.Broadcast(MaterialsThatUseThisTexture[i]);
+	}
+}
+
+#endif //WITH_EDITOR
\ No newline at end of file
diff --git a/Engine/Source/Runtime/Engine/Private/TextureDerivedDataTask.h b/Engine/Source/Runtime/Engine/Private/TextureDerivedDataTask.h
--- a/Engine/Source/Runtime/Engine/Private/TextureDerivedDataTask.h
+++ b/Engine/Source/Runtime/Engine/Private/TextureDerivedDataTask.h
@@ -13,7 +13,7 @@
 #include "Async/AsyncWork.h"
 #include "ImageCore.h"
 #include "TextureCompressorModule.h"
-#include "Interfaces/IImageWrapperModule.h"
+#include "IImageWrapperModule.h"
 #include "Engine/Texture2D.h"
 
 #endif // WITH_EDITOR
diff --git a/Engine/Source/Runtime/Engine/Private/TextureRenderTarget.cpp b/Engine/Source/Runtime/Engine/Private/TextureRenderTarget.cpp
--- a/Engine/Source/Runtime/Engine/Private/TextureRenderTarget.cpp
+++ b/Engine/Source/Runtime/Engine/Private/TextureRenderTarget.cpp
@@ -51,7 +51,7 @@ FTextureResource* UTextureRenderTarget::CreateResource()
 }
 
 
-EMaterialValueType UTextureRenderTarget::GetMaterialType()
+EMaterialValueType UTextureRenderTarget::GetMaterialType() const
 {
 	return MCT_Texture;
 }
diff --git a/Engine/Source/Runtime/Engine/Private/TextureRenderTarget2D.cpp b/Engine/Source/Runtime/Engine/Private/TextureRenderTarget2D.cpp
--- a/Engine/Source/Runtime/Engine/Private/TextureRenderTarget2D.cpp
+++ b/Engine/Source/Runtime/Engine/Private/TextureRenderTarget2D.cpp
@@ -59,7 +59,7 @@ FTextureResource* UTextureRenderTarget2D::CreateResource()
 	return Result;
 }
 
-EMaterialValueType UTextureRenderTarget2D::GetMaterialType()
+EMaterialValueType UTextureRenderTarget2D::GetMaterialType() const
 {
 	return MCT_Texture2D;
 }
diff --git a/Engine/Source/Runtime/Engine/Private/TextureRenderTargetCube.cpp b/Engine/Source/Runtime/Engine/Private/TextureRenderTargetCube.cpp
--- a/Engine/Source/Runtime/Engine/Private/TextureRenderTargetCube.cpp
+++ b/Engine/Source/Runtime/Engine/Private/TextureRenderTargetCube.cpp
@@ -78,7 +78,7 @@ FTextureResource* UTextureRenderTargetCube::CreateResource()
 }
 
 
-EMaterialValueType UTextureRenderTargetCube::GetMaterialType()
+EMaterialValueType UTextureRenderTargetCube::GetMaterialType() const
 {
 	return MCT_TextureCube;
 }
diff --git a/Engine/Source/Runtime/Engine/Private/UnrealEngine.cpp b/Engine/Source/Runtime/Engine/Private/UnrealEngine.cpp
--- a/Engine/Source/Runtime/Engine/Private/UnrealEngine.cpp
+++ b/Engine/Source/Runtime/Engine/Private/UnrealEngine.cpp
@@ -95,7 +95,8 @@
 #include "ShaderCompiler.h"
 #include "Slate/SlateSoundDevice.h"
 #include "DerivedDataCacheInterface.h"
-#include "Interfaces/IImageWrapperModule.h"
+#include "IImageWrapper.h"
+#include "IImageWrapperModule.h"
 #include "EngineAnalytics.h"
 #include "TickTaskManagerInterface.h"
 #include "Net/NetworkProfiler.h"
@@ -186,12 +187,12 @@
 #include "GameplayTagsManager.h"
 
 #if !UE_BUILD_SHIPPING
-#include "Interfaces/IAutomationWorkerModule.h"
-#include "HAL/ExceptionHandling.h"
+	#include "HAL/ExceptionHandling.h"
+	#include "IAutomationWorkerModule.h"
 #endif	// UE_BUILD_SHIPPING
 
 #if ENABLE_LOC_TESTING
-#include "LocalizationModule.h"
+	#include "LocalizationModule.h"
 #endif
 
 #include "GeneralProjectSettings.h"
@@ -200,9 +201,9 @@
 #include "AssetRegistryModule.h"
 
 #if !UE_BUILD_SHIPPING
-#include "IPluginManager.h"
-#include "GenericPlatformCrashContext.h"
-#include "EngineBuildSettings.h"
+	#include "IPluginManager.h"
+	#include "GenericPlatformCrashContext.h"
+	#include "EngineBuildSettings.h"
 #endif
 
 DEFINE_LOG_CATEGORY(LogEngine);
@@ -210,6 +211,7 @@ IMPLEMENT_MODULE( FEngineModule, Engine );
 
 #define LOCTEXT_NAMESPACE "UnrealEngine"
 
+
 void OnChangeEngineCVarRequiringRecreateRenderState(IConsoleVariable* Var)
 {
 	// Propgate cvar change to static draw lists
@@ -650,7 +652,7 @@ public:
 		if( InUncompressedData.Num() > 0 )
 		{
 			IImageWrapperModule& ImageWrapperModule = FModuleManager::LoadModuleChecked<IImageWrapperModule>( FName("ImageWrapper") );
-			IImageWrapperPtr ImageWrapper = ImageWrapperModule.CreateImageWrapper( EImageFormat::PNG );
+			TSharedPtr<IImageWrapper> ImageWrapper = ImageWrapperModule.CreateImageWrapper( EImageFormat::PNG );
 			if ( ImageWrapper.IsValid() && ImageWrapper->SetRaw( &InUncompressedData[ 0 ], InUncompressedData.Num(), InWidth, InHeight, ERGBFormat::RGBA, 8 ) )
 			{
 				OutCompressedData = ImageWrapper->GetCompressed();
@@ -679,7 +681,7 @@ public:
 		if( InCompressedData.Num() > 0 )
 		{
 			IImageWrapperModule& ImageWrapperModule = FModuleManager::LoadModuleChecked<IImageWrapperModule>( FName("ImageWrapper") );
-			IImageWrapperPtr ImageWrapper = ImageWrapperModule.CreateImageWrapper( EImageFormat::PNG );
+			TSharedPtr<IImageWrapper> ImageWrapper = ImageWrapperModule.CreateImageWrapper( EImageFormat::PNG );
 			if ( ImageWrapper.IsValid() && ImageWrapper->SetCompressed( &InCompressedData[ 0 ], InCompressedData.Num() ) )
 			{
 				check( ImageWrapper->GetWidth() == InWidth );
@@ -695,8 +697,6 @@ public:
 
 		return bSucceeded;
 	}
-
-
 };
 
 
diff --git a/Engine/Source/Runtime/Engine/Public/ExternalTexture.h b/Engine/Source/Runtime/Engine/Public/ExternalTexture.h
--- a/Engine/Source/Runtime/Engine/Public/ExternalTexture.h
+++ b/Engine/Source/Runtime/Engine/Public/ExternalTexture.h
@@ -18,13 +18,17 @@ class ENGINE_API FExternalTextureRegistry
 
 	struct FExternalTextureEntry
 	{
-		FExternalTextureEntry(FTextureRHIRef& InTextureRHI, FSamplerStateRHIRef& InSamplerStateRHI)
+		FExternalTextureEntry(FTextureRHIRef& InTextureRHI, FSamplerStateRHIRef& InSamplerStateRHI, const FLinearColor& InCoordinateScaleRotation, const FLinearColor& InCoordinateOffset)
 			: TextureRHI(InTextureRHI)
 			, SamplerStateRHI(InSamplerStateRHI)
+			, CoordinateScaleRotation(InCoordinateScaleRotation)
+			, CoordinateOffset(InCoordinateOffset)
 		{}
 
 		const FTextureRHIRef TextureRHI;
 		const FSamplerStateRHIRef SamplerStateRHI;
+		FLinearColor CoordinateScaleRotation;
+		FLinearColor CoordinateOffset;
 	};
 
 	TMap<FGuid, FExternalTextureEntry> TextureEntries;
@@ -35,16 +39,25 @@ public:
 	static FExternalTextureRegistry& Get();
 
 	/* Register an external texture, its sampler state and coordinate scale/bias against a GUID */
-	void RegisterExternalTexture(const FGuid& InGuid, FTextureRHIRef& InTextureRHI, FSamplerStateRHIRef& InSamplerStateRHI);
+	void RegisterExternalTexture(const FGuid& InGuid, FTextureRHIRef& InTextureRHI, FSamplerStateRHIRef& InSamplerStateRHI, const FLinearColor& InCoordinateScaleRotation = FLinearColor(1, 0, 0, 1), const FLinearColor& InCoordinateOffset = FLinearColor(0, 0, 0, 0));
 
 	/* Removes an external texture given a GUID */
 	void UnregisterExternalTexture(const FGuid& InGuid);
 
-	/* Removes reference to a MaterialRenderProxy which was referencing external textures */
 	void RemoveMaterialRenderProxyReference(const FMaterialRenderProxy* MaterialRenderProxy);
 
 	/* Looks up an external texture for given a given GUID
 	 * @return false if the texture is not registered
 	 */
 	bool GetExternalTexture(const FMaterialRenderProxy* MaterialRenderProxy, const FGuid& InGuid, FTextureRHIRef& OutTextureRHI, FSamplerStateRHIRef& OutSamplerStateRHI);
+
+	/* Looks up an texture coordinate scale rotation for given a given GUID
+	 * @return false if the texture is not registered
+	 */
+	bool GetExternalTextureCoordinateScaleRotation(const FGuid& InGuid, FLinearColor& OutCoordinateScaleRotation);
+
+	/* Looks up an texture coordinate offset for given a given GUID
+	 * @return false if the texture is not registered
+	 */
+	bool GetExternalTextureCoordinateOffset(const FGuid& InGuid, FLinearColor& OutCoordinateOffset);
 };
diff --git a/Engine/Source/Runtime/Engine/Public/MaterialCompiler.h b/Engine/Source/Runtime/Engine/Public/MaterialCompiler.h
--- a/Engine/Source/Runtime/Engine/Public/MaterialCompiler.h
+++ b/Engine/Source/Runtime/Engine/Public/MaterialCompiler.h
@@ -155,7 +155,14 @@ public:
 
 	virtual int32 Texture(UTexture* Texture,int32& TextureReferenceIndex,ESamplerSourceMode SamplerSource=SSM_FromTextureAsset,ETextureMipValueMode MipValueMode=TMVM_None) = 0;
 	virtual int32 TextureParameter(FName ParameterName,UTexture* DefaultTexture,int32& TextureReferenceIndex,ESamplerSourceMode SamplerSource=SSM_FromTextureAsset) = 0;
+
 	virtual int32 ExternalTexture(const FGuid& ExternalTextureGuid) = 0;
+	virtual int32 ExternalTexture(UTexture* InTexture, int32& TextureReferenceIndex) = 0;
+	virtual int32 ExternalTextureParameter(FName ParameterName, UTexture* DefaultValue, int32& TextureReferenceIndex) = 0;
+	virtual int32 ExternalTextureCoordinateScaleRotation(int32 TextureReferenceIndex, TOptional<FName> ParameterName) = 0;
+	virtual int32 ExternalTextureCoordinateScaleRotation(const FGuid& ExternalTextureGuid) = 0;
+	virtual int32 ExternalTextureCoordinateOffset(int32 TextureReferenceIndex, TOptional<FName> ParameterName) = 0;
+	virtual int32 ExternalTextureCoordinateOffset(const FGuid& ExternalTextureGuid) = 0;
 
 	virtual int32 GetTextureReferenceIndex(UTexture* Texture) { return INDEX_NONE; }
 
@@ -165,12 +172,24 @@ public:
 		return Texture(InTexture, TextureReferenceIndex, SamplerSource);
 	}
 
+	int32 ExternalTexture(UTexture* DefaultTexture)
+	{
+		int32 TextureReferenceIndex = INDEX_NONE;
+		return ExternalTexture(DefaultTexture, TextureReferenceIndex);
+	}
+
 	int32 TextureParameter(FName ParameterName,UTexture* DefaultTexture,ESamplerSourceMode SamplerSource=SSM_FromTextureAsset)
 	{
 		int32 TextureReferenceIndex = INDEX_NONE;
 		return  TextureParameter(ParameterName, DefaultTexture, TextureReferenceIndex, SamplerSource);
 	}
 
+	int32 ExternalTextureParameter(FName ParameterName, UTexture* DefaultTexture)
+	{
+		int32 TextureReferenceIndex = INDEX_NONE;
+		return ExternalTextureParameter(ParameterName, DefaultTexture, TextureReferenceIndex);
+	}
+
 	virtual	int32 PixelDepth()=0;
 	virtual int32 SceneDepth(int32 Offset, int32 UV, bool bUseOffset) = 0;
 	virtual int32 SceneColor(int32 Offset, int32 UV, bool bUseOffset) = 0;
@@ -256,6 +275,7 @@ public:
 	virtual int32 DistanceFieldGradient(int32 PositionArg) = 0;
 	virtual int32 DepthOfFieldFunction(int32 Depth, int32 FunctionValueIndex) = 0;
 	virtual int32 AtmosphericFogColor(int32 WorldPosition) = 0;
+	virtual int32 RotateScaleOffsetTexCoords(int32 TexCoordCodeIndex, int32 RotationScale, int32 Offset) = 0;
 	virtual int32 SpeedTree(ESpeedTreeGeometryType GeometryType, ESpeedTreeWindType WindType, ESpeedTreeLODType LODType, float BillboardThreshold, bool bAccurateWindVelocities) = 0;
 	virtual int32 TextureCoordinateOffset() = 0;
 	virtual int32 EyeAdaptation() = 0;
@@ -374,7 +394,14 @@ public:
 
 	virtual int32 Texture(UTexture* InTexture,int32& TextureReferenceIndex,ESamplerSourceMode SamplerSource=SSM_FromTextureAsset,ETextureMipValueMode MipValueMode=TMVM_None) override { return Compiler->Texture(InTexture,TextureReferenceIndex,SamplerSource,MipValueMode); }
 	virtual int32 TextureParameter(FName ParameterName,UTexture* DefaultValue,int32& TextureReferenceIndex,ESamplerSourceMode SamplerSource=SSM_FromTextureAsset) override { return Compiler->TextureParameter(ParameterName,DefaultValue,TextureReferenceIndex,SamplerSource); }
+
 	virtual int32 ExternalTexture(const FGuid& ExternalTextureGuid) override { return Compiler->ExternalTexture(ExternalTextureGuid); }
+	virtual int32 ExternalTexture(UTexture* InTexture, int32& TextureReferenceIndex) override { return Compiler->ExternalTexture(InTexture, TextureReferenceIndex); }
+	virtual int32 ExternalTextureParameter(FName ParameterName, UTexture* DefaultValue, int32& TextureReferenceIndex) override { return Compiler->ExternalTextureParameter(ParameterName, DefaultValue, TextureReferenceIndex); }
+	virtual int32 ExternalTextureCoordinateScaleRotation(int32 TextureReferenceIndex, TOptional<FName> ParameterName) override { return Compiler->ExternalTextureCoordinateScaleRotation(TextureReferenceIndex, ParameterName); }
+	virtual int32 ExternalTextureCoordinateScaleRotation(const FGuid& ExternalTextureGuid) override { return Compiler->ExternalTextureCoordinateScaleRotation(ExternalTextureGuid); }
+	virtual int32 ExternalTextureCoordinateOffset(int32 TextureReferenceIndex, TOptional<FName> ParameterName) override { return Compiler->ExternalTextureCoordinateOffset(TextureReferenceIndex, ParameterName); }
+	virtual int32 ExternalTextureCoordinateOffset(const FGuid& ExternalTextureGuid) override { return Compiler->ExternalTextureCoordinateOffset(ExternalTextureGuid); }
 
 	virtual	int32 PixelDepth() override { return Compiler->PixelDepth();	}
 	virtual int32 SceneDepth(int32 Offset, int32 UV, bool bUseOffset) override { return Compiler->SceneDepth(Offset, UV, bUseOffset); }
@@ -470,6 +497,11 @@ public:
 		return Compiler->DepthOfFieldFunction(Depth, FunctionValueIndex);
 	}
 
+	virtual int32 RotateScaleOffsetTexCoords(int32 TexCoordCodeIndex, int32 RotationScale, int32 Offset) override
+	{
+		return Compiler->RotateScaleOffsetTexCoords(TexCoordCodeIndex, RotationScale, Offset);
+	}
+
 	virtual int32 SpeedTree(ESpeedTreeGeometryType GeometryType, ESpeedTreeWindType WindType, ESpeedTreeLODType LODType, float BillboardThreshold, bool bAccurateWindVelocities) override 
 	{ 
 		return Compiler->SpeedTree(GeometryType, WindType, LODType, BillboardThreshold, bAccurateWindVelocities); 
diff --git a/Engine/Source/Runtime/Engine/Public/MaterialShared.h b/Engine/Source/Runtime/Engine/Public/MaterialShared.h
--- a/Engine/Source/Runtime/Engine/Public/MaterialShared.h
+++ b/Engine/Source/Runtime/Engine/Public/MaterialShared.h
@@ -275,39 +275,51 @@ protected:
 	UTexture* TransientOverrideValue_RenderThread;
 };
 
+class FMaterialUniformExpressionExternalTextureBase : public FMaterialUniformExpression
+{
+	DECLARE_MATERIALUNIFORMEXPRESSION_TYPE(FMaterialUniformExpressionExternalTextureBase);
+public:
+
+	FMaterialUniformExpressionExternalTextureBase(int32 InSourceTextureIndex = INDEX_NONE);
+	FMaterialUniformExpressionExternalTextureBase(const FGuid& InExternalTextureGuid);
+
+	virtual void Serialize(FArchive& Ar) override;
+	virtual bool IsConstant() const override { return false; }
+	virtual bool IsIdentical(const FMaterialUniformExpression* OtherExpression) const override;
+
+protected:
+
+	/** Resolve the guid that relates to texture information inside FExternalTexture */
+	FGuid ResolveExternalTextureGUID(const FMaterialRenderContext& Context, TOptional<FName> ParameterName = TOptional<FName>()) const;
+
+	/** Index of the texture in the material that should be used to retrieve the external texture GUID at runtime (or INDEX_NONE) */
+	int32 SourceTextureIndex;
+	/** Optional external texture GUID defined at compile time */
+	FGuid ExternalTextureGuid;
+};
 
 /**
 * An external texture expression.
 */
-class FMaterialUniformExpressionExternalTexture : public FMaterialUniformExpression
+class FMaterialUniformExpressionExternalTexture : public FMaterialUniformExpressionExternalTextureBase
 {
 	DECLARE_MATERIALUNIFORMEXPRESSION_TYPE(FMaterialUniformExpressionExternalTexture);
 public:
 
-	FMaterialUniformExpressionExternalTexture();
-
-	FMaterialUniformExpressionExternalTexture(const FGuid& InGuid);
+	FMaterialUniformExpressionExternalTexture(int32 InSourceTextureIndex = INDEX_NONE) : FMaterialUniformExpressionExternalTextureBase(InSourceTextureIndex) {}
+	FMaterialUniformExpressionExternalTexture(const FGuid& InGuid) : FMaterialUniformExpressionExternalTextureBase(InGuid) {}
 
 	// FMaterialUniformExpression interface.
-	virtual void Serialize(FArchive& Ar);
-	virtual class FMaterialUniformExpressionExternalTexture* GetExternalTextureUniformExpression() override { return this; }
+	virtual FMaterialUniformExpressionExternalTexture* GetExternalTextureUniformExpression() override { return this; }
 
 	// Lookup the external texture if it is set
-	bool GetExternalTexture(const FMaterialRenderProxy* MaterialRenderProxy, FTextureRHIRef& OutTextureRHI, FSamplerStateRHIRef& OutSamplerStateRHI);
+	virtual bool GetExternalTexture(const FMaterialRenderContext& Context, FTextureRHIRef& OutTextureRHI, FSamplerStateRHIRef& OutSamplerStateRHI) const;
 
-	virtual bool IsConstant() const
+	friend FArchive& operator<<(FArchive& Ar, FMaterialUniformExpressionExternalTexture*& Ref)
 	{
-		return false;
+		Ar << (FMaterialUniformExpression*&)Ref;
+		return Ar;
 	}
-	virtual bool IsIdentical(const FMaterialUniformExpression* OtherExpression) const;
-
-	friend FArchive& operator<<(FArchive& Ar, class FMaterialUniformExpressionExternalTexture*& Ref);
-
-	FGuid GetTextureExternalTextureGuid() const { return ExternalTextureGuid; }
-
-protected:
-	/** GUID is key for the ExternalTextures map */
-	FGuid ExternalTextureGuid;
 };
 
 /** Stores all uniform expressions for a material generated from a material translation. */
@@ -333,6 +345,7 @@ public:
 			+ UniformScalarExpressions.GetAllocatedSize()
 			+ Uniform2DTextureExpressions.GetAllocatedSize()
 			+ UniformCubeTextureExpressions.GetAllocatedSize()
+			+ UniformExternalTextureExpressions.GetAllocatedSize()
 			+ PerFrameUniformScalarExpressions.GetAllocatedSize()
 			+ PerFrameUniformVectorExpressions.GetAllocatedSize()
 			+ PerFramePrevUniformScalarExpressions.GetAllocatedSize()
diff --git a/Engine/Source/Runtime/Engine/Public/SceneView.h b/Engine/Source/Runtime/Engine/Public/SceneView.h
--- a/Engine/Source/Runtime/Engine/Public/SceneView.h
+++ b/Engine/Source/Runtime/Engine/Public/SceneView.h
@@ -234,7 +234,7 @@ struct FViewMatrices
 		ScreenScale = 1.f;
 	}
 
-	FViewMatrices(const FSceneViewInitOptions& InitOptions);
+	ENGINE_API FViewMatrices(const FSceneViewInitOptions& InitOptions);
 
 private:
 	/** ViewToClip : UE4 projection matrix projects such that clip space Z=1 is the near plane, and Z=0 is the infinite far plane. */
diff --git a/Engine/Source/Runtime/Engine/Public/ShowFlagsValues.inl b/Engine/Source/Runtime/Engine/Public/ShowFlagsValues.inl
--- a/Engine/Source/Runtime/Engine/Public/ShowFlagsValues.inl
+++ b/Engine/Source/Runtime/Engine/Public/ShowFlagsValues.inl
@@ -324,5 +324,7 @@ SHOWFLAG_FIXED_IN_SHIPPING(0, RequiredTextureResolution, SFG_Hidden, LOCTEXT("Re
 SHOWFLAG_ALWAYS_ACCESSIBLE(WidgetComponents, SFG_Normal, LOCTEXT("WidgetComponentsSF", "Widget Components"))
 /** Draw the bones of all skeletal meshes */
 SHOWFLAG_FIXED_IN_SHIPPING(0, Bones, SFG_Developer, LOCTEXT("BoneSF", "Bones"))
+/** If media planes should be shown */
+SHOWFLAG_ALWAYS_ACCESSIBLE(MediaPlanes, SFG_Normal, LOCTEXT("MediaPlanesSF", "Media Planes"))
 #undef SHOWFLAG_ALWAYS_ACCESSIBLE
 #undef SHOWFLAG_FIXED_IN_SHIPPING
diff --git a/Engine/Source/Runtime/Foliage/Private/ProceduralFoliageComponent.cpp b/Engine/Source/Runtime/Foliage/Private/ProceduralFoliageComponent.cpp
--- a/Engine/Source/Runtime/Foliage/Private/ProceduralFoliageComponent.cpp
+++ b/Engine/Source/Runtime/Foliage/Private/ProceduralFoliageComponent.cpp
@@ -213,7 +213,7 @@ bool UProceduralFoliageComponent::ExecuteSimulation(TArray<FDesiredFoliageInstan
 			for (int Y = 0; Y < TileLayout.NumTilesY; ++Y)
 			{
 				bool bFirstTime = true;
-				while (Futures[FutureIdx].WaitFor(FTimespan(0, 0, 0, 0, 100)) == false || bFirstTime)
+				while (Futures[FutureIdx].WaitFor(FTimespan::FromMilliseconds(100.0)) == false || bFirstTime)
 				{
 					if (GWarn->ReceivedUserCancel() && bCancelled == false)
 					{
diff --git a/Engine/Source/Runtime/Foliage/Private/ProceduralFoliageSpawner.cpp b/Engine/Source/Runtime/Foliage/Private/ProceduralFoliageSpawner.cpp
--- a/Engine/Source/Runtime/Foliage/Private/ProceduralFoliageSpawner.cpp
+++ b/Engine/Source/Runtime/Foliage/Private/ProceduralFoliageSpawner.cpp
@@ -165,7 +165,7 @@ void UProceduralFoliageSpawner::Simulate(int32 NumSteps)
 	for (int32 FutureIdx = 0; FutureIdx < Futures.Num(); ++FutureIdx)
 	{
 		// Sleep for 100ms if not ready. Needed so cancel is responsive.
-		while (Futures[FutureIdx].WaitFor(FTimespan(0, 0, 0, 0, 100)) == false)
+		while (Futures[FutureIdx].WaitFor(FTimespan::FromMilliseconds(100.0)) == false)
 		{
 			GWarn->StatusUpdate(FutureIdx, TotalTasks, LOCTEXT("SimulateProceduralFoliage", "Simulate ProceduralFoliage..."));
 
diff --git a/Engine/Source/Runtime/ImageWrapper/ImageWrapper.Build.cs b/Engine/Source/Runtime/ImageWrapper/ImageWrapper.Build.cs
--- a/Engine/Source/Runtime/ImageWrapper/ImageWrapper.Build.cs
+++ b/Engine/Source/Runtime/ImageWrapper/ImageWrapper.Build.cs
@@ -6,10 +6,14 @@ public class ImageWrapper : ModuleRules
 {
 	public ImageWrapper(ReadOnlyTargetRules Target) : base(Target)
 	{
-		PrivateIncludePaths.Add("Runtime/ImageWrapper/Private");
+		PrivateIncludePaths.AddRange(
+			new string[] {
+				"Runtime/ImageWrapper/Private",
+				"Runtime/ImageWrapper/Private/Formats",
+			});
 
 		Definitions.Add("WITH_UNREALPNG=1");
-        Definitions.Add("WITH_UNREALJPEG=1");
+		Definitions.Add("WITH_UNREALJPEG=1");
 
 		PrivateDependencyModuleNames.Add("Core");
 
@@ -17,17 +21,17 @@ public class ImageWrapper : ModuleRules
 			"zlib",
 			"UElibPNG",
 			"UElibJPG"
-			);
-
-        // Add openEXR lib for windows builds.
-        if (Target.Platform == UnrealTargetPlatform.Win64 ||
-			Target.Platform == UnrealTargetPlatform.Win32 ||
-			Target.Platform == UnrealTargetPlatform.Mac ||
-			Target.Platform == UnrealTargetPlatform.Linux)
-        {
-            Definitions.Add("WITH_UNREALEXR=1");
-            AddEngineThirdPartyPrivateStaticDependencies(Target, "UEOpenExr");
-        }
+		);
+
+		// Add openEXR lib for windows builds.
+		if ((Target.Platform == UnrealTargetPlatform.Win64) ||
+			(Target.Platform == UnrealTargetPlatform.Win32) ||
+			(Target.Platform == UnrealTargetPlatform.Mac) ||
+			(Target.Platform == UnrealTargetPlatform.Linux))
+		{
+			Definitions.Add("WITH_UNREALEXR=1");
+			AddEngineThirdPartyPrivateStaticDependencies(Target, "UEOpenExr");
+		}
 		else
 		{
 			Definitions.Add("WITH_UNREALEXR=0");
@@ -40,6 +44,6 @@ public class ImageWrapper : ModuleRules
 
 		// Disable shared PCHs to handle warning C4652
 		PCHUsage = ModuleRules.PCHUsageMode.NoSharedPCHs;
-		PrivatePCHHeaderFile = "Private/ImageWrapperPrivatePCH.h";
-    }
+		PrivatePCHHeaderFile = "Private/ImageWrapperPrivate.h";
+	}
 }
diff --git a/Engine/Source/Runtime/ImageWrapper/Private/Formats/BmpImageWrapper.cpp b/Engine/Source/Runtime/ImageWrapper/Private/Formats/BmpImageWrapper.cpp
new file mode 100644
--- /dev/null
+++ b/Engine/Source/Runtime/ImageWrapper/Private/Formats/BmpImageWrapper.cpp
@@ -0,0 +1,255 @@
+// Copyright 1998-2017 Epic Games, Inc. All Rights Reserved.
+
+#include "BmpImageWrapper.h"
+#include "ImageWrapperPrivate.h"
+
+#include "BmpImageSupport.h"
+
+
+/**
+ * BMP image wrapper class.
+ * This code was adapted from UTextureFactory::ImportTexture, but has not been throughly tested.
+ */
+
+FBmpImageWrapper::FBmpImageWrapper(bool bInHasHeader, bool bInHalfHeight)
+	: FImageWrapperBase()
+	, bHasHeader(bInHasHeader)
+	, bHalfHeight(bInHalfHeight)
+{
+}
+
+
+void FBmpImageWrapper::Compress(int32 Quality)
+{
+	checkf(false, TEXT("BMP compression not supported"));
+}
+
+
+void FBmpImageWrapper::Uncompress(const ERGBFormat InFormat, const int32 InBitDepth)
+{
+	const uint8* Buffer = CompressedData.GetData();
+
+	if (!bHasHeader || ((CompressedData.Num()>=sizeof(FBitmapFileHeader)+sizeof(FBitmapInfoHeader)) && Buffer[0]=='B' && Buffer[1]=='M'))
+	{
+		UncompressBMPData(InFormat, InBitDepth);
+	}
+}
+
+
+void FBmpImageWrapper::UncompressBMPData(const ERGBFormat InFormat, const int32 InBitDepth)
+{
+	const uint8* Buffer = CompressedData.GetData();
+	const FBitmapInfoHeader* bmhdr = nullptr;
+	const uint8* Bits = nullptr;
+
+	if (bHasHeader)
+	{
+		bmhdr = (FBitmapInfoHeader *)(Buffer + sizeof(FBitmapFileHeader));
+		Bits = Buffer + ((FBitmapFileHeader *)Buffer)->bfOffBits;
+	}
+	else
+	{
+		bmhdr = (FBitmapInfoHeader *)Buffer;
+		Bits = Buffer + sizeof(FBitmapInfoHeader);
+	}
+
+	if (bmhdr->biCompression != BCBI_RGB)
+	{
+		UE_LOG(LogImageWrapper, Error, TEXT("RLE compression of BMP images not supported"));
+		return;
+	}
+
+	if (bmhdr->biPlanes==1 && bmhdr->biBitCount==8)
+	{
+		// Do palette.
+		const uint8* bmpal = (uint8*)CompressedData.GetData() + sizeof(FBitmapFileHeader) + sizeof(FBitmapInfoHeader);
+
+		// Set texture properties.
+		Width = bmhdr->biWidth;
+		const bool bNegativeHeight = (bmhdr->biHeight < 0);
+		Height = FMath::Abs(bHalfHeight ? bmhdr->biHeight / 2 : bmhdr->biHeight);
+		Format = ERGBFormat::BGRA;
+		RawData.Empty(Height * Width * 4);
+		RawData.AddUninitialized(Height * Width * 4);
+
+		FColor* ImageData = (FColor*)RawData.GetData();
+
+		// If the number for color palette entries is 0, we need to default to 2^biBitCount entries.  In this case 2^8 = 256
+		int32 clrPaletteCount = bmhdr->biClrUsed ? bmhdr->biClrUsed : 256;
+		TArray<FColor>	Palette;
+
+		for (int32 i = 0; i < clrPaletteCount; i++)
+		{
+			Palette.Add(FColor(bmpal[i * 4 + 2], bmpal[i * 4 + 1], bmpal[i * 4 + 0], 255));
+		}
+
+		while (Palette.Num() < 256)
+		{
+			Palette.Add(FColor(0, 0, 0, 255));
+		}
+
+		// Copy scanlines, accounting for scanline direction according to the Height field.
+		const int32 SrcStride = Align(Width, 4);
+		const int32 SrcPtrDiff = bNegativeHeight ? SrcStride : -SrcStride;
+		const uint8* SrcPtr = Bits + (bNegativeHeight ? 0 : Height - 1) * SrcStride;
+
+		for (int32 Y = 0; Y < Height; Y++)
+		{
+			for (int32 X = 0; X < Width; X++)
+			{
+				*ImageData++ = Palette[SrcPtr[X]];
+			}
+
+			SrcPtr += SrcPtrDiff;
+		}
+	}
+	else if (bmhdr->biPlanes==1 && bmhdr->biBitCount==24)
+	{
+		// Set texture properties.
+		Width = bmhdr->biWidth;
+		const bool bNegativeHeight = (bmhdr->biHeight < 0);
+		Height = FMath::Abs(bHalfHeight ? bmhdr->biHeight / 2 : bmhdr->biHeight);
+		Format = ERGBFormat::BGRA;
+		RawData.Empty(Height * Width * 4);
+		RawData.AddUninitialized(Height * Width * 4);
+
+		uint8* ImageData = RawData.GetData();
+
+		// Copy scanlines, accounting for scanline direction according to the Height field.
+		const int32 SrcStride = Align(Width * 3, 4);
+		const int32 SrcPtrDiff = bNegativeHeight ? SrcStride : -SrcStride;
+		const uint8* SrcPtr = Bits + (bNegativeHeight ? 0 : Height - 1) * SrcStride;
+
+		for (int32 Y = 0; Y < Height; Y++)
+		{
+			const uint8* SrcRowPtr = SrcPtr;
+			for (int32 X = 0; X < Width; X++)
+			{
+				*ImageData++ = *SrcRowPtr++;
+				*ImageData++ = *SrcRowPtr++;
+				*ImageData++ = *SrcRowPtr++;
+				*ImageData++ = 0xFF;
+			}
+
+			SrcPtr += SrcPtrDiff;
+		}
+	}
+	else if (bmhdr->biPlanes==1 && bmhdr->biBitCount==32)
+	{
+		// Set texture properties.
+		Width = bmhdr->biWidth;
+		const bool bNegativeHeight = (bmhdr->biHeight < 0);
+		Height = FMath::Abs(bHalfHeight ? bmhdr->biHeight / 2 : bmhdr->biHeight);
+		Format = ERGBFormat::BGRA;
+		RawData.Empty(Height * Width * 4);
+		RawData.AddUninitialized(Height * Width * 4);
+
+		uint8* ImageData = RawData.GetData();
+
+		// Copy scanlines, accounting for scanline direction according to the Height field.
+		const int32 SrcStride = Width * 4;
+		const int32 SrcPtrDiff = bNegativeHeight ? SrcStride : -SrcStride;
+		const uint8* SrcPtr = Bits + (bNegativeHeight ? 0 : Height - 1) * SrcStride;
+
+		for (int32 Y = 0; Y < Height; Y++)
+		{
+			const uint8* SrcRowPtr = SrcPtr;
+			for (int32 X = 0; X < Width; X++)
+			{
+				*ImageData++ = *SrcRowPtr++;
+				*ImageData++ = *SrcRowPtr++;
+				*ImageData++ = *SrcRowPtr++;
+				*ImageData++ = *SrcRowPtr++;
+			}
+
+			SrcPtr += SrcPtrDiff;
+		}
+	}
+	else if (bmhdr->biPlanes==1 && bmhdr->biBitCount==16)
+	{
+		UE_LOG(LogImageWrapper, Error, TEXT("BMP 16 bit format no longer supported. Use terrain tools for importing/exporting heightmaps."));
+	}
+	else
+	{
+		UE_LOG(LogImageWrapper, Error, TEXT("BMP uses an unsupported format (%i/%i)"), bmhdr->biPlanes, bmhdr->biBitCount);
+	}
+}
+
+
+bool FBmpImageWrapper::SetCompressed(const void* InCompressedData, int32 InCompressedSize)
+{
+	bool bResult = FImageWrapperBase::SetCompressed(InCompressedData, InCompressedSize);
+
+	return bResult && (bHasHeader ? LoadBMPHeader() : LoadBMPInfoHeader());	// Fetch the variables from the header info
+}
+
+
+bool FBmpImageWrapper::LoadBMPHeader()
+{
+	const FBitmapInfoHeader* bmhdr = (FBitmapInfoHeader *)(CompressedData.GetData() + sizeof(FBitmapFileHeader));
+	const FBitmapFileHeader* bmf   = (FBitmapFileHeader *)(CompressedData.GetData() + 0);
+	if ((CompressedData.Num() >= sizeof(FBitmapFileHeader) + sizeof(FBitmapInfoHeader)) && CompressedData.GetData()[0] == 'B' && CompressedData.GetData()[1] == 'M')
+	{
+		if (bmhdr->biCompression != BCBI_RGB)
+		{
+			UE_LOG(LogImageWrapper, Error, TEXT("RLE compression of BMP images not supported"));
+			return false;
+		}
+
+		if (bmhdr->biPlanes==1 && (bmhdr->biBitCount==8 || bmhdr->biBitCount==24 || bmhdr->biBitCount==32))
+		{
+			// Set texture properties.
+			Width = bmhdr->biWidth;
+			Height = FMath::Abs(bmhdr->biHeight);
+			Format = ERGBFormat::BGRA;
+			BitDepth = bmhdr->biBitCount;
+
+			return true;
+		}
+		
+		if (bmhdr->biPlanes == 1 && bmhdr->biBitCount == 16)
+		{
+			UE_LOG(LogImageWrapper, Error, TEXT("BMP 16 bit format no longer supported. Use terrain tools for importing/exporting heightmaps."));
+		}
+		else
+		{
+			UE_LOG(LogImageWrapper, Error, TEXT("BMP uses an unsupported format (%i/%i)"), bmhdr->biPlanes, bmhdr->biBitCount);
+		}
+	}
+
+	return false;
+}
+
+
+bool FBmpImageWrapper::LoadBMPInfoHeader()
+{
+	const FBitmapInfoHeader* bmhdr = (FBitmapInfoHeader *)CompressedData.GetData();
+
+	if (bmhdr->biCompression != BCBI_RGB)
+	{
+		UE_LOG(LogImageWrapper, Error, TEXT("RLE compression of BMP images not supported"));
+		return false;
+	}
+
+	if (bmhdr->biPlanes==1 && (bmhdr->biBitCount==8 || bmhdr->biBitCount==24 || bmhdr->biBitCount==32))
+	{
+		// Set texture properties.
+		Width = bmhdr->biWidth;
+		Height = FMath::Abs(bmhdr->biHeight);
+		Format = ERGBFormat::BGRA;
+		BitDepth = bmhdr->biBitCount;
+
+		return true;
+	}
+	
+	if (bmhdr->biPlanes == 1 && bmhdr->biBitCount == 16)
+	{
+		UE_LOG(LogImageWrapper, Error, TEXT("BMP 16 bit format no longer supported. Use terrain tools for importing/exporting heightmaps."));
+	}
+	else
+	{
+		UE_LOG(LogImageWrapper, Error, TEXT("BMP uses an unsupported format (%i/%i)"), bmhdr->biPlanes, bmhdr->biBitCount);
+	}
+
+	return false;
+}
diff --git a/Engine/Source/Runtime/ImageWrapper/Private/Formats/ExrImageWrapper.cpp b/Engine/Source/Runtime/ImageWrapper/Private/Formats/ExrImageWrapper.cpp
new file mode 100644
--- /dev/null
+++ b/Engine/Source/Runtime/ImageWrapper/Private/Formats/ExrImageWrapper.cpp
@@ -0,0 +1,442 @@
+// Copyright 1998-2017 Epic Games, Inc. All Rights Reserved.
+
+#include "ExrImageWrapper.h"
+#include "ImageWrapperPrivate.h"
+
+#include "HAL/PlatformTime.h"
+#include "Math/Float16.h"
+
+
+#if WITH_UNREALEXR
+
+FExrImageWrapper::FExrImageWrapper()
+	: FImageWrapperBase()
+	, bUseCompression(true)
+{
+}
+
+
+template <typename sourcetype> class FSourceImageRaw
+{
+public:
+	FSourceImageRaw(const TArray<uint8>& SourceImageBitmapIN, uint32 ChannelsIN, uint32 WidthIN, uint32 HeightIN)
+		: SourceImageBitmap(SourceImageBitmapIN),
+			Width( WidthIN ),
+			Height(HeightIN),
+			Channels(ChannelsIN)
+	{		
+		check(SourceImageBitmap.Num() == Channels*Width*Height*sizeof(sourcetype));
+	}
+
+	uint32 GetXStride() const	{return sizeof(sourcetype) * Channels;}
+	uint32 GetYStride() const	{return GetXStride() * Width;}
+
+	const sourcetype* GetHorzLine(uint32 y, uint32 Channel)
+	{
+		return &SourceImageBitmap[(sizeof(sourcetype) * Channel) + (GetYStride() * y)];
+	}
+
+private:
+	const TArray<uint8>& SourceImageBitmap;
+	uint32 Width;
+	uint32 Height;
+	uint32 Channels;
+};
+
+
+class FMemFileOut : public Imf::OStream
+{
+public:
+	//-------------------------------------------------------
+	// A constructor that opens the file with the given name.
+	// The destructor will close the file.
+	//-------------------------------------------------------
+
+	FMemFileOut(const char fileName[]) :
+		Imf::OStream(fileName),
+		Pos(0)
+	{
+	}
+
+	virtual void write(const char c[/*n*/], int n)
+	{
+		for (int32 i = 0; i < n; ++i)
+		{
+			Data[Pos + i] = c[i];
+		}
+		Pos += n;
+	}
+
+
+	//---------------------------------------------------------
+	// Get the current writing position, in bytes from the
+	// beginning of the file.  If the next call to write() will
+	// start writing at the beginning of the file, tellp()
+	// returns 0.
+	//---------------------------------------------------------
+
+	virtual Imf::Int64 tellp()
+	{
+		return Pos;
+	}
+
+
+	//-------------------------------------------
+	// Set the current writing position.
+	// After calling seekp(i), tellp() returns i.
+	//-------------------------------------------
+
+	virtual void seekp(Imf::Int64 pos)
+	{
+		Pos = pos;
+	}
+
+
+	int64 Pos;
+	TArray<uint8> Data;
+};
+
+
+class FMemFileIn : public Imf::IStream
+{
+public:
+	//-------------------------------------------------------
+	// A constructor that opens the file with the given name.
+	// The destructor will close the file.
+	//-------------------------------------------------------
+
+	FMemFileIn(const void* InData, int32 InSize)
+		: Imf::IStream("")
+		, Data((const char *)InData)
+		, Size(InSize)
+		, Pos(0)
+	{
+	}
+
+	//------------------------------------------------------
+	// Read from the stream:
+	//
+	// read(c,n) reads n bytes from the stream, and stores
+	// them in array c.  If the stream contains less than n
+	// bytes, or if an I/O error occurs, read(c,n) throws
+	// an exception.  If read(c,n) reads the last byte from
+	// the file it returns false, otherwise it returns true.
+	//------------------------------------------------------
+
+	virtual bool read (char c[/*n*/], int n)
+	{
+		if(Pos + n > Size)
+		{
+			return false;
+		}
+
+		for (int32 i = 0; i < n; ++i)
+		{
+			c[i] = Data[Pos];
+			++Pos;
+		}
+
+		return Pos >= Size;
+	}
+
+	//--------------------------------------------------------
+	// Get the current reading position, in bytes from the
+	// beginning of the file.  If the next call to read() will
+	// read the first byte in the file, tellg() returns 0.
+	//--------------------------------------------------------
+
+	virtual Imf::Int64 tellg()
+	{
+		return Pos;
+	}
+
+	//-------------------------------------------
+	// Set the current reading position.
+	// After calling seekg(i), tellg() returns i.
+	//-------------------------------------------
+
+	virtual void seekg(Imf::Int64 pos)
+	{
+		Pos = pos;
+	}
+
+private:
+
+	const char* Data;
+	int32 Size;
+	int64 Pos;
+};
+
+
+namespace
+{
+	/////////////////////////////////////////
+	// 8 bit per channel source
+	float ToLinear(uint8 LDRValue)
+	{
+		return pow((float)(LDRValue) / 255.f, 2.2f);
+	}
+
+	void ExtractAndConvertChannel(const uint8*Src, uint32 SrcChannels, uint32 x, uint32 y, float* ChannelOUT)
+	{
+		for (uint32 i = 0; i < x*y; i++, Src += SrcChannels)
+		{
+			ChannelOUT[i] = ToLinear(*Src);
+		}
+	}
+
+	void ExtractAndConvertChannel(const uint8*Src, uint32 SrcChannels, uint32 x, uint32 y, FFloat16* ChannelOUT)
+	{
+		for (uint32 i = 0; i < x*y; i++, Src += SrcChannels)
+		{
+			ChannelOUT[i] = FFloat16(ToLinear(*Src));
+		}
+	}
+
+	/////////////////////////////////////////
+	// 16 bit per channel source
+	void ExtractAndConvertChannel(const FFloat16*Src, uint32 SrcChannels, uint32 x, uint32 y, float* ChannelOUT)
+	{
+		for (uint32 i = 0; i < x*y; i++, Src += SrcChannels)
+		{
+			ChannelOUT[i] = Src->GetFloat();
+		}
+	}
+
+	void ExtractAndConvertChannel(const FFloat16*Src, uint32 SrcChannels, uint32 x, uint32 y, FFloat16* ChannelOUT)
+	{
+		for (uint32 i = 0; i < x*y; i++, Src += SrcChannels)
+		{
+			ChannelOUT[i] = *Src;
+		}
+	}
+
+	/////////////////////////////////////////
+	// 32 bit per channel source
+	void ExtractAndConvertChannel(const float* Src, uint32 SrcChannels, uint32 x, uint32 y, float* ChannelOUT)
+	{
+		for (uint32 i = 0; i < x*y; i++, Src += SrcChannels)
+		{
+			ChannelOUT[i] = *Src;
+		}
+	}
+
+	void ExtractAndConvertChannel(const float* Src, uint32 SrcChannels, uint32 x, uint32 y, FFloat16* ChannelOUT)
+	{
+		for (uint32 i = 0; i < x*y; i++, Src += SrcChannels)
+		{
+			ChannelOUT[i] = *Src;
+		}
+	}
+
+	/////////////////////////////////////////
+	int32 GetNumChannelsFromFormat(ERGBFormat Format)
+	{
+		switch (Format)
+		{
+			case ERGBFormat::RGBA:
+			case ERGBFormat::BGRA:
+				return 4;
+			case ERGBFormat::Gray:
+				return 1;
+		}
+		checkNoEntry();
+		return 1;
+	}
+}
+
+const char* FExrImageWrapper::GetRawChannelName(int ChannelIndex) const
+{
+	const int32 MaxChannels = 4;
+	static const char* RGBAChannelNames[] = { "R", "G", "B", "A" };
+	static const char* BGRAChannelNames[] = { "B", "G", "R", "A" };
+	static const char* GrayChannelNames[] = { "G" };
+	check(ChannelIndex < MaxChannels);
+
+	const char** ChannelNames = BGRAChannelNames;
+
+	switch (RawFormat)
+	{
+		case ERGBFormat::RGBA:
+		{
+			ChannelNames = RGBAChannelNames;
+		}
+		break;
+		case ERGBFormat::BGRA:
+		{
+			ChannelNames = BGRAChannelNames;
+		}
+		break;
+		case ERGBFormat::Gray:
+		{
+			check(ChannelIndex < ARRAY_COUNT(GrayChannelNames));
+			ChannelNames = GrayChannelNames;
+		}
+		break;
+		default:
+			checkNoEntry();
+	}
+	return ChannelNames[ChannelIndex];
+}
+
+template <typename Imf::PixelType OutputFormat>
+struct TExrImageOutputChannelType;
+
+template <> struct TExrImageOutputChannelType<Imf::HALF>  { typedef FFloat16 Type; };
+template <> struct TExrImageOutputChannelType<Imf::FLOAT> { typedef float Type; };
+
+template <Imf::PixelType OutputFormat, typename sourcetype>
+void FExrImageWrapper::WriteFrameBufferChannel(Imf::FrameBuffer& ImfFrameBuffer, const char* ChannelName, const sourcetype* SrcData, TArray<uint8>& ChannelBuffer)
+{
+	const int32 OutputPixelSize = ((OutputFormat == Imf::HALF) ? 2 : 4);
+	ChannelBuffer.AddUninitialized(Width*Height*OutputPixelSize);
+	uint32 SrcChannels = GetNumChannelsFromFormat(RawFormat);
+	ExtractAndConvertChannel(SrcData, SrcChannels, Width, Height, (typename TExrImageOutputChannelType<OutputFormat>::Type*)&ChannelBuffer[0]);
+	Imf::Slice FrameChannel = Imf::Slice(OutputFormat, (char*)&ChannelBuffer[0], OutputPixelSize, Width*OutputPixelSize);
+	ImfFrameBuffer.insert(ChannelName, FrameChannel);
+}
+
+template <Imf::PixelType OutputFormat, typename sourcetype>
+void FExrImageWrapper::CompressRaw(const sourcetype* SrcData, bool bIgnoreAlpha)
+{
+	const double StartTime = FPlatformTime::Seconds();
+	uint32 NumWriteComponents = GetNumChannelsFromFormat(RawFormat);
+	if (bIgnoreAlpha && NumWriteComponents == 4)
+	{
+		NumWriteComponents = 3;
+	}
+
+	Imf::Compression Comp = bUseCompression ? Imf::Compression::ZIP_COMPRESSION : Imf::Compression::NO_COMPRESSION;
+	Imf::Header Header(Width, Height, 1, Imath::V2f(0, 0), 1, Imf::LineOrder::INCREASING_Y, Comp);
+	
+	for (uint32 Channel = 0; Channel < NumWriteComponents; Channel++)
+	{
+		Header.channels().insert(GetRawChannelName(Channel), Imf::Channel(OutputFormat));
+	}
+
+	// OutputFormat is 0 UINT, 1 HALF or 2 FLOAT. Size evaluates to 2/4/8 or 1/2/4 compressed
+	const int32 OutputPixelSize = (bUseCompression ? 1 : 2) << OutputFormat;
+	checkSlow(OutputPixelSize >= 1 && OutputPixelSize <= 8);
+
+	FMemFileOut MemFile("");
+	MemFile.Data.AddUninitialized(Width * Height * NumWriteComponents * OutputPixelSize);
+
+	Imf::FrameBuffer ImfFrameBuffer;
+	TArray<uint8> ChannelOutputBuffers[4];
+
+	for (uint32 Channel = 0; Channel < NumWriteComponents; Channel++)
+	{
+		WriteFrameBufferChannel<OutputFormat>(ImfFrameBuffer, GetRawChannelName(Channel), SrcData + Channel, ChannelOutputBuffers[Channel]);
+	}
+
+	Imf::OutputFile ImfFile(MemFile, Header);
+	ImfFile.setFrameBuffer(ImfFrameBuffer);
+	ImfFile.writePixels(Height);
+
+	CompressedData.AddUninitialized(MemFile.tellp());
+	FMemory::Memcpy(CompressedData.GetData(), MemFile.Data.GetData(), MemFile.tellp());
+
+	const double DeltaTime = FPlatformTime::Seconds() - StartTime;
+	UE_LOG(LogImageWrapper, Verbose, TEXT("Compressed image in %.3f seconds"), DeltaTime);
+}
+
+void FExrImageWrapper::Compress( int32 Quality )
+{
+	check(RawData.Num() != 0);
+	check(Width > 0);
+	check(Height > 0);
+	check(RawBitDepth == 8 || RawBitDepth == 16 || RawBitDepth == 32);
+
+	const int32 MaxComponents = 4;
+
+	bUseCompression = (Quality != (int32)EImageCompressionQuality::Uncompressed);
+
+	switch (RawBitDepth)
+	{
+		case 8:
+			CompressRaw<Imf::HALF>(&RawData[0], false);
+		break;
+		case 16:
+			CompressRaw<Imf::HALF>((const FFloat16*)&RawData[0], false);
+		break;
+		case 32:
+			CompressRaw<Imf::FLOAT>((const float*)&RawData[0], false);
+		break;
+		default:
+			checkNoEntry();
+	}
+}
+
+void FExrImageWrapper::Uncompress( const ERGBFormat InFormat, const int32 InBitDepth )
+{
+	// Ensure we haven't already uncompressed the file.
+	if ( RawData.Num() != 0 )
+	{
+		return;
+	}
+
+	FMemFileIn MemFile(&CompressedData[0], CompressedData.Num());
+
+	Imf::RgbaInputFile ImfFile(MemFile);
+
+	Imath::Box2i win = ImfFile.dataWindow();
+
+	check(BitDepth == 16);
+	check(Width);
+	check(Height);
+
+	uint32 Channels = 4;
+
+	RawData.Empty();
+	RawData.AddUninitialized( Width * Height * Channels * (BitDepth / 8) );
+
+	int dx = win.min.x;
+	int dy = win.min.y;
+
+	ImfFile.setFrameBuffer((Imf::Rgba*)&RawData[0] - dx - dy * Width, 1, Width);
+	ImfFile.readPixels(win.min.y, win.max.y);
+}
+
+// from http://www.openexr.com/ReadingAndWritingImageFiles.pdf
+bool IsThisAnOpenExrFile(Imf::IStream& f)
+{
+	char b[4];
+	f.read(b, sizeof(b));
+
+	f.seekg(0);
+
+	return b[0] == 0x76 && b[1] == 0x2f && b[2] == 0x31 && b[3] == 0x01;
+}
+
+bool FExrImageWrapper::SetCompressed( const void* InCompressedData, int32 InCompressedSize )
+{
+	if(!FImageWrapperBase::SetCompressed( InCompressedData, InCompressedSize))
+	{
+		return false;
+	}
+
+	FMemFileIn MemFile(InCompressedData, InCompressedSize);
+
+	if(!IsThisAnOpenExrFile(MemFile))
+	{
+		return false;
+	}
+
+	Imf::RgbaInputFile ImfFile(MemFile);
+
+	Imath::Box2i win = ImfFile.dataWindow();
+
+	Imath::V2i dim(win.max.x - win.min.x + 1, win.max.y - win.min.y + 1);
+
+	BitDepth = 16;
+
+	Width = dim.x;
+	Height = dim.y;
+
+	// ideally we can specify float here
+	Format = ERGBFormat::RGBA;
+
+	return true;
+}
+
+#endif // WITH_UNREALEXR
diff --git a/Engine/Source/Runtime/ImageWrapper/Private/Formats/ExrImageWrapper.h b/Engine/Source/Runtime/ImageWrapper/Private/Formats/ExrImageWrapper.h
new file mode 100644
--- /dev/null
+++ b/Engine/Source/Runtime/ImageWrapper/Private/Formats/ExrImageWrapper.h
@@ -0,0 +1,61 @@
+// Copyright 1998-2017 Epic Games, Inc. All Rights Reserved.
+
+#pragma once
+
+#include "Containers/Array.h"
+
+#include "ImageWrapperBase.h"
+
+#if WITH_UNREALEXR
+
+THIRD_PARTY_INCLUDES_START
+	#include "ThirdParty/openexr/Deploy/include/ImfIO.h"
+	#include "ThirdParty/openexr/Deploy/include/ImathBox.h"
+	#include "ThirdParty/openexr/Deploy/include/ImfChannelList.h"
+	#include "ThirdParty/openexr/Deploy/include/ImfInputFile.h"
+	#include "ThirdParty/openexr/Deploy/include/ImfOutputFile.h"
+	#include "ThirdParty/openexr/Deploy/include/ImfArray.h"
+	#include "ThirdParty/openexr/Deploy/include/ImfHeader.h"
+	#include "ThirdParty/openexr/Deploy/include/ImfStdIO.h"
+	#include "ThirdParty/openexr/Deploy/include/ImfChannelList.h"
+	#include "ThirdParty/openexr/Deploy/include/ImfRgbaFile.h"
+THIRD_PARTY_INCLUDES_END
+
+
+/**
+ * OpenEXR implementation of the helper class
+ */
+class FExrImageWrapper
+	: public FImageWrapperBase
+{
+public:
+
+	/**
+	 * Default Constructor.
+	 */
+	FExrImageWrapper();
+
+public:
+
+	//~ FImageWrapper interface
+
+	virtual void Compress(int32 Quality) override;
+	virtual void Uncompress(const ERGBFormat InFormat, int32 InBitDepth) override;
+	virtual bool SetCompressed(const void* InCompressedData, int32 InCompressedSize) override;
+
+protected:
+
+	template <Imf::PixelType OutputFormat, typename sourcetype>
+	void WriteFrameBufferChannel(Imf::FrameBuffer& ImfFrameBuffer, const char* ChannelName, const sourcetype* SrcData, TArray<uint8>& ChannelBuffer);
+
+	template <Imf::PixelType OutputFormat, typename sourcetype>
+	void CompressRaw(const sourcetype* SrcData, bool bIgnoreAlpha);
+
+	const char* GetRawChannelName(int ChannelIndex) const;
+
+private:
+
+	bool bUseCompression;
+};
+
+#endif // WITH_UNREALEXR
diff --git a/Engine/Source/Runtime/ImageWrapper/Private/Formats/IcnsImageWrapper.cpp b/Engine/Source/Runtime/ImageWrapper/Private/Formats/IcnsImageWrapper.cpp
new file mode 100644
--- /dev/null
+++ b/Engine/Source/Runtime/ImageWrapper/Private/Formats/IcnsImageWrapper.cpp
@@ -0,0 +1,82 @@
+// Copyright 1998-2017 Epic Games, Inc. All Rights Reserved.
+
+#include "IcnsImageWrapper.h"
+
+
+/* FIcnsImageWrapper structors
+ *****************************************************************************/
+
+FIcnsImageWrapper::FIcnsImageWrapper()
+	: FImageWrapperBase()
+{ }
+
+
+/* FImageWrapper interface
+ *****************************************************************************/
+
+bool FIcnsImageWrapper::SetCompressed(const void* InCompressedData, int32 InCompressedSize)
+{
+#if PLATFORM_MAC
+	return FImageWrapperBase::SetCompressed(InCompressedData, InCompressedSize);
+#else
+	return false;
+#endif
+}
+
+
+bool FIcnsImageWrapper::SetRaw(const void* InRawData, int32 InRawSize, const int32 InWidth, const int32 InHeight, const ERGBFormat InFormat, const int32 InBitDepth)
+{
+#if PLATFORM_MAC
+	return FImageWrapperBase::SetRaw(InRawData, InRawSize, InWidth, InHeight, InFormat, InBitDepth);
+#else
+	return false;
+#endif
+}
+
+
+void FIcnsImageWrapper::Compress(int32 Quality)
+{
+	checkf(false, TEXT("ICNS compression not supported"));
+}
+
+
+void FIcnsImageWrapper::Uncompress(const ERGBFormat InFormat, const int32 InBitDepth)
+{
+#if PLATFORM_MAC
+	SCOPED_AUTORELEASE_POOL;
+
+	NSData* ImageData = [NSData dataWithBytesNoCopy:CompressedData.GetData() length:CompressedData.Num() freeWhenDone:NO];
+	NSImage* Image = [[NSImage alloc] initWithData:ImageData];
+	if (Image)
+	{
+		NSBitmapImageRep* Bitmap = [NSBitmapImageRep imageRepWithData:[Image TIFFRepresentation]];
+		if (Bitmap)
+		{
+			check(InFormat == ERGBFormat::BGRA || InFormat == ERGBFormat::RGBA);
+			check(InBitDepth == 8);
+
+			RawData.Empty();
+			RawData.Append([Bitmap bitmapData], [Bitmap bytesPerPlane]);
+
+			RawFormat = Format = InFormat;
+			RawBitDepth = BitDepth = InBitDepth;
+
+			Width = [Bitmap pixelsWide];
+			Height = [Bitmap pixelsHigh];
+
+			if (Format == ERGBFormat::BGRA)
+			{
+				for (int32 Index = 0; Index < [Bitmap bytesPerPlane]; Index += 4)
+				{
+					uint8 Byte = RawData[Index];
+					RawData[Index] = RawData[Index + 2];
+					RawData[Index + 2] = Byte;
+				}
+			}
+		}
+		[Image release];
+	}
+#else
+	checkf(false, TEXT("ICNS uncompressing not supported on this platform"));
+#endif
+}
diff --git a/Engine/Source/Runtime/ImageWrapper/Private/Formats/IcnsImageWrapper.h b/Engine/Source/Runtime/ImageWrapper/Private/Formats/IcnsImageWrapper.h
new file mode 100644
--- /dev/null
+++ b/Engine/Source/Runtime/ImageWrapper/Private/Formats/IcnsImageWrapper.h
@@ -0,0 +1,27 @@
+// Copyright 1998-2017 Epic Games, Inc. All Rights Reserved.
+
+#pragma once
+
+#include "ImageWrapperBase.h"
+
+
+/**
+ * ICNS implementation of the helper class.
+ */
+class FIcnsImageWrapper
+	: public FImageWrapperBase
+{
+public:
+
+	/** Default Constructor. */
+	FIcnsImageWrapper();
+
+public:
+
+	//~ FImageWrapper Interface
+
+	virtual bool SetCompressed(const void* InCompressedData, int32 InCompressedSize) override;
+	virtual bool SetRaw(const void* InRawData, int32 InRawSize, const int32 InWidth, const int32 InHeight, const ERGBFormat InFormat, const int32 InBitDepth) override;
+	virtual void Compress(int32 Quality) override;
+	virtual void Uncompress(const ERGBFormat InFormat, int32 InBitDepth) override;
+};
diff --git a/Engine/Source/Runtime/ImageWrapper/Private/Formats/JpegImageWrapper.cpp b/Engine/Source/Runtime/ImageWrapper/Private/Formats/JpegImageWrapper.cpp
new file mode 100644
--- /dev/null
+++ b/Engine/Source/Runtime/ImageWrapper/Private/Formats/JpegImageWrapper.cpp
@@ -0,0 +1,176 @@
+// Copyright 1998-2017 Epic Games, Inc. All Rights Reserved.
+
+#include "JpegImageWrapper.h"
+
+#include "Math/Color.h"
+#include "Misc/ScopeLock.h"
+
+#if WITH_UNREALJPEG
+
+#ifdef __clang__
+	#pragma clang diagnostic push
+	#pragma clang diagnostic ignored "-Wshadow"
+
+	#if PLATFORM_LINUX || PLATFORM_MAC
+		#pragma clang diagnostic ignored "-Wshift-negative-value"	// clang 3.7.0
+	#endif
+#endif
+
+PRAGMA_DISABLE_SHADOW_VARIABLE_WARNINGS
+#include "jpgd.h"
+#include "jpgd.cpp"
+#include "jpge.h"
+#include "jpge.cpp"
+PRAGMA_ENABLE_SHADOW_VARIABLE_WARNINGS
+
+#ifdef __clang__
+	#pragma clang diagnostic pop
+#endif
+
+DEFINE_LOG_CATEGORY_STATIC(JPEGLog, Log, All);
+
+
+// Only allow one thread to use JPEG decoder at a time (it's not thread safe)
+FCriticalSection GJPEGSection; 
+
+
+/* FJpegImageWrapper structors
+ *****************************************************************************/
+
+FJpegImageWrapper::FJpegImageWrapper(int32 InNumComponents)
+	: FImageWrapperBase(),
+	  NumComponents(InNumComponents)
+{ }
+
+
+/* FImageWrapperBase interface
+ *****************************************************************************/
+
+bool FJpegImageWrapper::SetCompressed(const void* InCompressedData, int32 InCompressedSize)
+{
+	jpgd::jpeg_decoder_mem_stream jpeg_memStream((uint8*)InCompressedData, InCompressedSize);
+
+	jpgd::jpeg_decoder decoder(&jpeg_memStream);
+	if (decoder.get_error_code() != jpgd::JPGD_SUCCESS)
+		return false;
+
+	bool bResult = FImageWrapperBase::SetCompressed(InCompressedData, InCompressedSize);
+
+	// We don't support 16 bit jpegs
+	BitDepth = 8;
+
+	Width = decoder.get_width();
+	Height = decoder.get_height();
+
+	switch (decoder.get_num_components())
+	{
+	case 1:
+		Format = ERGBFormat::Gray;
+		break;
+	case 3:
+		Format = ERGBFormat::RGBA;
+		break;
+	default:
+		return false;
+	}
+
+	return bResult;
+}
+
+
+bool FJpegImageWrapper::SetRaw(const void* InRawData, int32 InRawSize, const int32 InWidth, const int32 InHeight, const ERGBFormat InFormat, const int32 InBitDepth)
+{
+	check((InFormat == ERGBFormat::RGBA || InFormat == ERGBFormat::BGRA || InFormat == ERGBFormat::Gray) && InBitDepth == 8);
+
+	bool bResult = FImageWrapperBase::SetRaw(InRawData, InRawSize, InWidth, InHeight, InFormat, InBitDepth);
+
+	return bResult;
+}
+
+
+void FJpegImageWrapper::Compress(int32 Quality)
+{
+	if (CompressedData.Num() == 0)
+	{
+		FScopeLock JPEGLock(&GJPEGSection);
+		
+		if (Quality == 0) {Quality = 85;}
+		ensure(Quality >= 1 && Quality <= 100);
+		Quality = FMath::Clamp(Quality, 1, 100);
+
+		check(RawData.Num());
+		check(Width > 0);
+		check(Height > 0);
+
+		// re-order components if required - JPEGs expect RGBA
+		if(RawFormat == ERGBFormat::BGRA)
+		{
+			FColor* Colors = (FColor*)RawData.GetData();
+			const int32 NumColors = RawData.Num() / 4;
+			for(int32 ColorIndex = 0; ColorIndex < NumColors; ColorIndex++)
+			{
+				uint8 Temp = Colors[ColorIndex].B;
+				Colors[ColorIndex].B = Colors[ColorIndex].R;
+				Colors[ColorIndex].R = Temp;
+			}
+		}
+
+		CompressedData.Empty();
+		CompressedData.AddUninitialized(RawData.Num());
+
+		int OutBufferSize = CompressedData.Num();
+
+		jpge::params Parameters;
+		Parameters.m_quality = Quality;
+		bool bSuccess = jpge::compress_image_to_jpeg_file_in_memory(
+			CompressedData.GetData(), OutBufferSize, Width, Height, NumComponents, RawData.GetData(), Parameters);
+		
+		check(bSuccess);
+
+		CompressedData.RemoveAt(OutBufferSize, CompressedData.Num() - OutBufferSize);
+	}
+}
+
+
+void FJpegImageWrapper::Uncompress(const ERGBFormat InFormat, int32 InBitDepth)
+{
+	// Ensure we haven't already uncompressed the file.
+	if (RawData.Num() != 0)
+	{
+		return;
+	}
+
+	// Get the number of channels we need to extract
+	int Channels = 0;
+	if ((InFormat == ERGBFormat::RGBA || InFormat == ERGBFormat::BGRA) && InBitDepth == 8)
+	{
+		Channels = 4;
+	}
+	else if (InFormat == ERGBFormat::Gray && InBitDepth == 8)
+	{
+		Channels = 1;
+	}
+	else
+	{
+		check(false);
+	}
+
+	FScopeLock JPEGLock(&GJPEGSection);
+
+	check(CompressedData.Num());
+
+	int32 NumColors;
+	uint8* OutData = jpgd::decompress_jpeg_image_from_memory(
+		CompressedData.GetData(), CompressedData.Num(), &Width, &Height, &NumColors, Channels);
+
+	RawData.Empty();
+	RawData.AddUninitialized(Width * Height * Channels);
+	if (OutData)
+	{
+		FMemory::Memcpy(RawData.GetData(), OutData, RawData.Num());
+		FMemory::Free(OutData);
+	}
+}
+
+
+#endif //WITH_JPEG
diff --git a/Engine/Source/Runtime/ImageWrapper/Private/Formats/JpegImageWrapper.h b/Engine/Source/Runtime/ImageWrapper/Private/Formats/JpegImageWrapper.h
new file mode 100644
--- /dev/null
+++ b/Engine/Source/Runtime/ImageWrapper/Private/Formats/JpegImageWrapper.h
@@ -0,0 +1,39 @@
+// Copyright 1998-2017 Epic Games, Inc. All Rights Reserved.
+
+#pragma once
+
+#include "CoreTypes.h"
+#include "ImageWrapperBase.h"
+
+#if WITH_UNREALJPEG
+
+
+/**
+ * Uncompresses JPEG data to raw 24bit RGB image that can be used by Unreal textures.
+ *
+ * Source code for JPEG decompression from http://code.google.com/p/jpeg-compressor/
+ */
+class FJpegImageWrapper
+	: public FImageWrapperBase
+{
+public:
+
+	/** Default constructor. */
+	FJpegImageWrapper(int32 InNumComponents = 4);
+
+public:
+
+	//~ FImageWrapperBase interface
+
+	virtual bool SetCompressed(const void* InCompressedData, int32 InCompressedSize) override;
+	virtual bool SetRaw(const void* InRawData, int32 InRawSize, const int32 InWidth, const int32 InHeight, const ERGBFormat InFormat, const int32 InBitDepth) override;
+	virtual void Uncompress(const ERGBFormat InFormat, int32 InBitDepth) override;
+	virtual void Compress(int32 Quality) override;
+
+private:
+
+	int32 NumComponents;
+};
+
+
+#endif //WITH_JPEG
diff --git a/Engine/Source/Runtime/ImageWrapper/Private/ImageWrapperBase.cpp b/Engine/Source/Runtime/ImageWrapper/Private/ImageWrapperBase.cpp
--- a/Engine/Source/Runtime/ImageWrapper/Private/ImageWrapperBase.cpp
+++ b/Engine/Source/Runtime/ImageWrapper/Private/ImageWrapperBase.cpp
@@ -6,7 +6,7 @@
 /* FImageWrapperBase structors
  *****************************************************************************/
 
-FImageWrapperBase::FImageWrapperBase( )
+FImageWrapperBase::FImageWrapperBase()
 	: RawFormat(ERGBFormat::Invalid)
 	, RawBitDepth(0)
 	, Format(ERGBFormat::Invalid)
@@ -19,7 +19,7 @@ FImageWrapperBase::FImageWrapperBase( )
 /* FImageWrapperBase interface
  *****************************************************************************/
 
-void FImageWrapperBase::Reset( )
+void FImageWrapperBase::Reset()
 {
 	LastError.Empty();
 
@@ -32,7 +32,7 @@ void FImageWrapperBase::Reset( )
 }
 
 
-void FImageWrapperBase::SetError( const TCHAR* ErrorMessage )
+void FImageWrapperBase::SetError(const TCHAR* ErrorMessage)
 {
 	LastError = ErrorMessage;
 }
@@ -50,7 +50,7 @@ const TArray<uint8>& FImageWrapperBase::GetCompressed(int32 Quality)
 }
 
 
-bool FImageWrapperBase::GetRaw( const ERGBFormat::Type InFormat, int32 InBitDepth, const TArray<uint8>*& OutRawData )
+bool FImageWrapperBase::GetRaw(const ERGBFormat InFormat, int32 InBitDepth, const TArray<uint8>*& OutRawData)
 {
 	LastError.Empty();
 	Uncompress(InFormat, InBitDepth);
@@ -64,9 +64,9 @@ bool FImageWrapperBase::GetRaw( const ERGBFormat::Type InFormat, int32 InBitDept
 }
 
 
-bool FImageWrapperBase::SetCompressed( const void* InCompressedData, int32 InCompressedSize )
+bool FImageWrapperBase::SetCompressed(const void* InCompressedData, int32 InCompressedSize)
 {
-	if( InCompressedSize > 0 && InCompressedData != nullptr )
+	if(InCompressedSize > 0 && InCompressedData != nullptr)
 	{
 		Reset();
 		RawData.Empty();			// Invalidates the raw data too
@@ -82,7 +82,7 @@ bool FImageWrapperBase::SetCompressed( const void* InCompressedData, int32 InCom
 }
 
 
-bool FImageWrapperBase::SetRaw( const void* InRawData, int32 InRawSize, const int32 InWidth, const int32 InHeight, const ERGBFormat::Type InFormat, const int32 InBitDepth )
+bool FImageWrapperBase::SetRaw(const void* InRawData, int32 InRawSize, const int32 InWidth, const int32 InHeight, const ERGBFormat InFormat, const int32 InBitDepth)
 {
 	check(InRawData != NULL);
 	check(InRawSize > 0);
diff --git a/Engine/Source/Runtime/ImageWrapper/Private/ImageWrapperBase.h b/Engine/Source/Runtime/ImageWrapper/Private/ImageWrapperBase.h
--- a/Engine/Source/Runtime/ImageWrapper/Private/ImageWrapperBase.h
+++ b/Engine/Source/Runtime/ImageWrapper/Private/ImageWrapperBase.h
@@ -2,8 +2,12 @@
 
 #pragma once
 
-#include "CoreMinimal.h"
-#include "Interfaces/IImageWrapper.h"
+#include "CoreTypes.h"
+#include "Containers/Array.h"
+#include "Containers/EnumAsByte.h"
+#include "Containers/UnrealString.h"
+#include "IImageWrapper.h"
+
 
 /**
  * The abstract helper class for handling the different image formats
@@ -13,10 +17,8 @@ class FImageWrapperBase
 {
 public:
 
-	/**
-	 * Default Constructor.
-	 */
-	FImageWrapperBase( );
+	/** Default Constructor. */
+	FImageWrapperBase();
 
 public:
 
@@ -37,57 +39,57 @@ public:
 	 *
 	 * @param Quality The compression quality.
 	 */
-	virtual void Compress( int32 Quality ) = 0;
+	virtual void Compress(int32 Quality) = 0;
 
 	/**
 	 * Resets the local variables.
 	 */
-	virtual void Reset( );
+	virtual void Reset();
 
 	/**
 	 * Sets last error message.
 	 *
 	 * @param ErrorMessage The error message to set.
 	 */
-	virtual void SetError( const TCHAR* ErrorMessage );
+	virtual void SetError(const TCHAR* ErrorMessage);
 
 	/**  
 	 * Function to uncompress our data 
 	 *
 	 * @param InFormat How we want to manipulate the RGB data
 	 */
-	virtual void Uncompress( const ERGBFormat::Type InFormat, int32 InBitDepth ) = 0;
+	virtual void Uncompress(const ERGBFormat InFormat, int32 InBitDepth) = 0;
 
 public:
 
-	// IImageWrapper interface
+	//~ IImageWrapper interface
 
-	virtual const TArray<uint8>& GetCompressed( int32 Quality = 0 ) override;
+	virtual const TArray<uint8>& GetCompressed(int32 Quality = 0) override;
 
-	virtual int32 GetBitDepth( ) const override
+	virtual int32 GetBitDepth() const override
 	{
 		return BitDepth;
 	}
 
-	virtual ERGBFormat::Type GetFormat() const override
+	virtual ERGBFormat GetFormat() const override
 	{
 		return Format;
 	}
 
-	virtual int32 GetHeight( ) const override
+	virtual int32 GetHeight() const override
 	{
 		return Height;
 	}
 
-	virtual bool GetRaw( const ERGBFormat::Type InFormat, int32 InBitDepth, const TArray<uint8>*& OutRawData ) override;
+	virtual bool GetRaw(const ERGBFormat InFormat, int32 InBitDepth, const TArray<uint8>*& OutRawData) override;
 
-	virtual int32 GetWidth( ) const override
+	virtual int32 GetWidth() const override
 	{
 		return Width;
 	}
 
-	virtual bool SetCompressed( const void* InCompressedData, int32 InCompressedSize ) override;
-	virtual bool SetRaw( const void* InRawData, int32 InRawSize, const int32 InWidth, const int32 InHeight, const ERGBFormat::Type InFormat, const int32 InBitDepth ) override;
+	virtual bool SetCompressed(const void* InCompressedData, int32 InCompressedSize) override;
+	virtual bool SetRaw(const void* InRawData, int32 InRawSize, const int32 InWidth, const int32 InHeight, const ERGBFormat InFormat, const int32 InBitDepth) override;
 
 protected:
 
@@ -96,11 +98,11 @@ protected:
 	TArray<uint8> CompressedData;
 
 	/** Format of the raw data */
-	TEnumAsByte<ERGBFormat::Type> RawFormat;
+	ERGBFormat RawFormat;
 	int8 RawBitDepth;
 
 	/** Format of the image */
-	TEnumAsByte<ERGBFormat::Type> Format;
+	ERGBFormat Format;
 
 	/** Bit depth of the image */
 	int8 BitDepth;
diff --git a/Engine/Source/Runtime/ImageWrapper/Private/ImageWrapperModule.cpp b/Engine/Source/Runtime/ImageWrapper/Private/ImageWrapperModule.cpp
--- a/Engine/Source/Runtime/ImageWrapper/Private/ImageWrapperModule.cpp
+++ b/Engine/Source/Runtime/ImageWrapper/Private/ImageWrapperModule.cpp
@@ -1,22 +1,24 @@
 // Copyright 1998-2017 Epic Games, Inc. All Rights Reserved.
 
-#include "CoreMinimal.h"
-#include "Modules/ModuleManager.h"
-#include "Interfaces/IImageWrapperModule.h"
-#include "JpegImageWrapper.h"
-#include "PngImageWrapper.h"
-#include "BmpImageWrapper.h"
-#include "ExrImageWrapper.h"
-#include "IcoImageWrapper.h"
-#include "IcnsImageWrapper.h"
 #include "ImageWrapperPrivate.h"
 
+#include "CoreTypes.h"
+#include "Modules/ModuleManager.h"
 
-DEFINE_LOG_CATEGORY(LogImageWrapper);
+#include "Formats/BmpImageWrapper.h"
+#include "Formats/ExrImageWrapper.h"
+#include "Formats/IcnsImageWrapper.h"
+#include "Formats/IcoImageWrapper.h"
+#include "Formats/JpegImageWrapper.h"
+#include "Formats/PngImageWrapper.h"
+#include "IImageWrapperModule.h"
 
 
-namespace {
+DEFINE_LOG_CATEGORY(LogImageWrapper);
 
+
+namespace
+{
 	static const uint8 IMAGE_MAGIC_PNG[]  = {0x89, 0x50, 0x4E, 0x47, 0x0D, 0x0A, 0x1A, 0x0A};
 	static const uint8 IMAGE_MAGIC_JPEG[] = {0xFF, 0xD8, 0xFF};
 	static const uint8 IMAGE_MAGIC_BMP[]  = {0x42, 0x4D};
@@ -31,28 +33,31 @@ namespace {
 		{
 			return false;
 		}
+
 		for (int32 I = 0; I < MagicCount; ++I)
 		{
-			if ( Content[I] != Magic[I] )
+			if (Content[I] != Magic[I])
 			{
 				return false;
 			}
 		}
+
 		return true;
 	}
 }
 
+
 /**
- * Image Wrapper module
+ * Image Wrapper module.
  */
 class FImageWrapperModule
 	: public IImageWrapperModule
 {
 public:
 
-	// IImageWrapperModule interface
+	//~ IImageWrapperModule interface
 
-	virtual IImageWrapperPtr CreateImageWrapper( const EImageFormat::Type InFormat ) override
+	virtual TSharedPtr<IImageWrapper> CreateImageWrapper(const EImageFormat InFormat) override
 	{
 		FImageWrapperBase* ImageWrapper = NULL;
 
@@ -64,10 +69,12 @@ public:
 			ImageWrapper = new FPngImageWrapper();
 			break;
 #endif	// WITH_UNREALPNG
+
 #if WITH_UNREALJPEG
 		case EImageFormat::JPEG:
 			ImageWrapper = new FJpegImageWrapper();
 			break;
+
 		case EImageFormat::GrayscaleJPEG:
 			ImageWrapper = new FJpegImageWrapper(1);
 			break;
@@ -80,6 +87,7 @@ public:
 		case EImageFormat::ICO:
 			ImageWrapper = new FIcoImageWrapper();
 			break;
+
 #if WITH_UNREALEXR
 		case EImageFormat::EXR:
 			ImageWrapper = new FExrImageWrapper();
@@ -88,6 +96,7 @@ public:
 		case EImageFormat::ICNS:
 			ImageWrapper = new FIcnsImageWrapper();
 			break;
+
 		default:
 			break;
 		}
@@ -95,9 +104,9 @@ public:
 		return MakeShareable(ImageWrapper);
 	}
 
-	virtual EImageFormat::Type DetectImageFormat( const void* CompressedData, int32 CompressedSize ) override
+	virtual EImageFormat DetectImageFormat(const void* CompressedData, int32 CompressedSize) override
 	{
-		EImageFormat::Type Format = EImageFormat::Invalid;
+		EImageFormat Format = EImageFormat::Invalid;
 		if (StartsWith((uint8*) CompressedData, CompressedSize, IMAGE_MAGIC_PNG))
 		{
 			Format = EImageFormat::PNG;
@@ -128,10 +137,10 @@ public:
 
 public:
 
-	// IModuleInterface interface
+	//~ IModuleInterface interface
 
-	virtual void StartupModule( ) override { }
-	virtual void ShutdownModule( ) override { }
+	virtual void StartupModule() override { }
+	virtual void ShutdownModule() override { }
 };
 
 
diff --git a/Engine/Source/Runtime/ImageWrapper/Private/ImageWrapperPrivate.h b/Engine/Source/Runtime/ImageWrapper/Private/ImageWrapperPrivate.h
--- a/Engine/Source/Runtime/ImageWrapper/Private/ImageWrapperPrivate.h
+++ b/Engine/Source/Runtime/ImageWrapper/Private/ImageWrapperPrivate.h
@@ -2,7 +2,6 @@
 
 #pragma once
 
-#include "CoreMinimal.h"
+#include "Logging/LogMacros.h"
 
 DECLARE_LOG_CATEGORY_EXTERN(LogImageWrapper, Log, All);
-
diff --git a/Engine/Source/Runtime/ImageWrapper/Public/IImageWrapper.h b/Engine/Source/Runtime/ImageWrapper/Public/IImageWrapper.h
new file mode 100644
--- /dev/null
+++ b/Engine/Source/Runtime/ImageWrapper/Public/IImageWrapper.h
@@ -0,0 +1,149 @@
+// Copyright 1998-2017 Epic Games, Inc. All Rights Reserved.
+
+#pragma once
+
+#include "CoreTypes.h"
+#include "Containers/Array.h"
+#include "Templates/SharedPointer.h"
+
+
+/**
+ * Enumerates the types of image formats this class can handle.
+ */
+enum class EImageFormat : int8
+{
+	/** Invalid or unrecognized format. */
+	Invalid = -1,
+
+	/** Portable Network Graphics. */
+	PNG = 0,
+
+	/** Joint Photographic Experts Group. */
+	JPEG,
+
+	/** Single channel JPEG. */
+	GrayscaleJPEG,	
+
+	/** Windows Bitmap. */
+	BMP,
+
+	/** Windows Icon resource. */
+	ICO,
+
+	/** OpenEXR (HDR) image file format. */
+	EXR,
+
+	/** Mac icon. */
+	ICNS
+};
+
+
+/**
+ * Enumerates the types of RGB formats this class can handle.
+ */
+enum class ERGBFormat : int8
+{
+	Invalid = -1,
+	RGBA =  0,
+	BGRA =  1,
+	Gray =  2,
+};
+
+
+/**
+ * Enumerates available image compression qualities.
+ */
+enum class EImageCompressionQuality : uint8
+{
+	Default = 0,
+	Uncompressed =  1,
+};
+
+
+/**
+ * Interface for image wrappers.
+ */
+class IImageWrapper
+{
+public:
+
+	/**  
+	 * Sets the compressed data.
+	 *
+	 * @param InCompressedData The memory address of the start of the compressed data.
+	 * @param InCompressedSize The size of the compressed data parsed.
+	 * @return true if data was the expected format.
+	 */
+	virtual bool SetCompressed(const void* InCompressedData, int32 InCompressedSize) = 0;
+
+	/**  
+	 * Sets the compressed data.
+	 *
+	 * @param InRawData The memory address of the start of the raw data.
+	 * @param InRawSize The size of the compressed data parsed.
+	 * @param InWidth The width of the image data.
+	 * @param InHeight the height of the image data.
+	 * @param InFormat the format the raw data is in, normally RGBA.
+	 * @param InBitDepth the bit-depth per channel, normally 8.
+	 * @return true if data was the expected format.
+	 */
+	virtual bool SetRaw(const void* InRawData, int32 InRawSize, const int32 InWidth, const int32 InHeight, const ERGBFormat InFormat, const int32 InBitDepth) = 0;
+
+	/**
+	 * Gets the compressed data.
+	 *
+	 * @return Array of the compressed data.
+	 */
+	virtual const TArray<uint8>& GetCompressed(int32 Quality = 0) = 0;
+
+	/**  
+	 * Gets the raw data.
+	 *
+	 * @param InFormat How we want to manipulate the RGB data.
+	 * @param InBitDepth The output bit-depth per channel, normally 8.
+	 * @param OutRawData Will contain the uncompressed raw data.
+	 * @return true on success, false otherwise.
+	 */
+	virtual bool GetRaw(const ERGBFormat InFormat, int32 InBitDepth, const TArray<uint8>*& OutRawData) = 0;
+
+	/** 
+	 * Gets the width of the image.
+	 *
+	 * @return Image width.
+	 * @see GetHeight
+	 */
+	virtual int32 GetWidth() const = 0;
+
+	/** 
+	 * Gets the height of the image.
+	 *
+	 * @return Image height.
+	 * @see GetWidth
+	 */
+	virtual int32 GetHeight() const = 0;
+
+	/** 
+	 * Gets the bit depth of the image.
+	 *
+	 * @return The bit depth per-channel of the image.
+	 */
+	virtual int32 GetBitDepth() const = 0;
+
+	/** 
+	 * Gets the format of the image.
+	 * Theoretically, this is the format it would be best to call GetRaw() with, if you support it.
+	 *
+	 * @return The format the image data is in
+	 */
+	virtual ERGBFormat GetFormat() const = 0;
+
+public:
+
+	/** Virtual destructor. */
+	virtual ~IImageWrapper() { }
+};
+
+
+/** Type definition for shared pointers to instances of IImageWrapper. */
+DEPRECATED(4.16, "IImageWrapperPtr is deprecated. Please use 'TSharedPtr<IImageWrapper>' instead!")
+typedef TSharedPtr<IImageWrapper> IImageWrapperPtr;
diff --git a/Engine/Source/Runtime/ImageWrapper/Public/IImageWrapperModule.h b/Engine/Source/Runtime/ImageWrapper/Public/IImageWrapperModule.h
new file mode 100644
--- /dev/null
+++ b/Engine/Source/Runtime/ImageWrapper/Public/IImageWrapperModule.h
@@ -0,0 +1,43 @@
+// Copyright 1998-2017 Epic Games, Inc. All Rights Reserved.
+
+#pragma once
+
+#include "CoreTypes.h"
+#include "Modules/ModuleInterface.h"
+#include "Templates/SharedPointer.h"
+
+class IImageWrapper;
+enum class EImageFormat : int8;
+
+
+/**
+ * Interface for image wrapper modules.
+ */
+class IImageWrapperModule
+	: public IModuleInterface
+{
+public:
+
+	/**  
+	 * Create a helper of a specific type 
+	 *
+	 * @param InFormat - The type of image we want to deal with
+	 * @return The helper base class to manage the data
+	 */
+	virtual TSharedPtr<IImageWrapper> CreateImageWrapper(const EImageFormat InFormat) = 0;
+
+	/**
+	 * Detect image format by looking at the first few bytes of the compressed image data.
+	 * You can call this method as soon as you have 8-16 bytes of compressed file content available.
+	 *
+	 * @param InCompressedData The raw image header.
+	 * @param InCompressedSize The size of InCompressedData.
+	 * @return the detected format or EImageFormat::Invalid if the method could not detect the image format.
+	 */
+	virtual EImageFormat DetectImageFormat(const void* InCompressedData, int32 InCompressedSize) = 0;
+
+public:
+
+	/** Virtual destructor. */
+	virtual ~IImageWrapperModule() { }
+};
diff --git a/Engine/Source/Runtime/Launch/Launch.Build.cs b/Engine/Source/Runtime/Launch/Launch.Build.cs
--- a/Engine/Source/Runtime/Launch/Launch.Build.cs
+++ b/Engine/Source/Runtime/Launch/Launch.Build.cs
@@ -12,6 +12,8 @@ public class Launch : ModuleRules
 		PrivateIncludePathModuleNames.AddRange(
 			new string[] {
 				"AutomationController",
+				"Media",
+                "MediaUtils",
 				"TaskGraph",
 			}
 		);
@@ -22,7 +24,6 @@ public class Launch : ModuleRules
 				"CoreUObject",
 				"Engine",
 				"InputCore",
-				"MediaAssets",
 				"MoviePlayer",
 				"Networking",
 				"PakFile",
@@ -124,7 +125,13 @@ public class Launch : ModuleRules
 
 		if (Target.bCompileAgainstEngine)
 		{
-			PrivateIncludePathModuleNames.Add("Messaging");
+			PrivateIncludePathModuleNames.AddRange(
+				new string[] {
+					"MessagingCommon",
+					"MediaAssets"
+				}
+			);
+
 			PublicDependencyModuleNames.Add("SessionServices");
 			PrivateIncludePaths.Add("Developer/DerivedDataCache/Public");
 
diff --git a/Engine/Source/Runtime/Launch/Private/IOS/GameLaunchDaemonMessageHandler.cpp b/Engine/Source/Runtime/Launch/Private/IOS/GameLaunchDaemonMessageHandler.cpp
--- a/Engine/Source/Runtime/Launch/Private/IOS/GameLaunchDaemonMessageHandler.cpp
+++ b/Engine/Source/Runtime/Launch/Private/IOS/GameLaunchDaemonMessageHandler.cpp
@@ -1,10 +1,16 @@
 // Copyright 1998-2017 Epic Games, Inc. All Rights Reserved.
 
 #include "GameLaunchDaemonMessageHandler.h"
+
 #include "HAL/PlatformProcess.h"
-#include "Helpers/MessageEndpointBuilder.h"
+#include "IMessageContext.h"
+#include "IOSMessageProtocol.h"
+#include "MessageEndpoint.h"
+#include "MessageEndpointBuilder.h"
+
 #import <TargetConditionals.h>
 
+
 void FGameLaunchDaemonMessageHandler::Init()
 {
 	MessageEndpoint = FMessageEndpoint::Builder("FGameLaunchDaemonMessageHandler")
@@ -17,6 +23,7 @@ void FGameLaunchDaemonMessageHandler::Init()
 	}
 }
 
+
 void FGameLaunchDaemonMessageHandler::Shutdown()
 {
 	if (MessageEndpoint.IsValid())
@@ -25,7 +32,8 @@ void FGameLaunchDaemonMessageHandler::Shutdown()
 	}
 }
 
-void FGameLaunchDaemonMessageHandler::HandlePingMessage(const FIOSLaunchDaemonPing& Message, const IMessageContextRef& Context)
+
+void FGameLaunchDaemonMessageHandler::HandlePingMessage(const FIOSLaunchDaemonPing& Message, const TSharedRef<IMessageContext, ESPMode::ThreadSafe>& Context)
 {
 	if (MessageEndpoint.IsValid())
 	{
@@ -42,7 +50,8 @@ void FGameLaunchDaemonMessageHandler::HandlePingMessage(const FIOSLaunchDaemonPi
 	}
 }
 
-void FGameLaunchDaemonMessageHandler::HandleLaunchRequest(const FIOSLaunchDaemonLaunchApp& Message, const IMessageContextRef& Context)
+
+void FGameLaunchDaemonMessageHandler::HandleLaunchRequest(const FIOSLaunchDaemonLaunchApp& Message, const TSharedRef<IMessageContext, ESPMode::ThreadSafe>& Context)
 {
 	FString Error;
 
diff --git a/Engine/Source/Runtime/Launch/Private/IOS/GameLaunchDaemonMessageHandler.h b/Engine/Source/Runtime/Launch/Private/IOS/GameLaunchDaemonMessageHandler.h
--- a/Engine/Source/Runtime/Launch/Private/IOS/GameLaunchDaemonMessageHandler.h
+++ b/Engine/Source/Runtime/Launch/Private/IOS/GameLaunchDaemonMessageHandler.h
@@ -2,9 +2,15 @@
 
 #pragma once
 
-#include "IMessageContext.h"
-#include "Helpers/MessageEndpoint.h"
-#include "IOSMessageProtocol.h"
+#include "Containers/UnrealString.h"
+#include "Templates/SharedPointer.h"
+
+class FMessageEndpoint;
+class IMessageContext;
+
+struct FIOSLaunchDaemonLaunchApp;
+struct FIOSLaunchDaemonPing;
+
 
 class FGameLaunchDaemonMessageHandler
 {
@@ -14,10 +20,12 @@ public:
 	void Shutdown();
 
 private:
-	void HandlePingMessage(const FIOSLaunchDaemonPing& Message, const IMessageContextRef& Context);
-	void HandleLaunchRequest(const FIOSLaunchDaemonLaunchApp& Message, const IMessageContextRef& Context);
 
+	void HandlePingMessage(const FIOSLaunchDaemonPing& Message, const TSharedRef<IMessageContext, ESPMode::ThreadSafe>& Context);
+	void HandleLaunchRequest(const FIOSLaunchDaemonLaunchApp& Message, const TSharedRef<IMessageContext, ESPMode::ThreadSafe>& Context);
+
+private:
 
-	FMessageEndpointPtr MessageEndpoint;
+	TSharedPtr<FMessageEndpoint, ESPMode::ThreadSafe> MessageEndpoint;
 	FString AppId;
 };
diff --git a/Engine/Source/Runtime/Launch/Private/LaunchEngineLoop.cpp b/Engine/Source/Runtime/Launch/Private/LaunchEngineLoop.cpp
--- a/Engine/Source/Runtime/Launch/Private/LaunchEngineLoop.cpp
+++ b/Engine/Source/Runtime/Launch/Private/LaunchEngineLoop.cpp
@@ -1,6 +1,7 @@
 // Copyright 1998-2017 Epic Games, Inc. All Rights Reserved.
 
 #include "LaunchEngineLoop.h"
+
 #include "HAL/PlatformStackWalk.h"
 #include "HAL/PlatformOutputDevices.h"
 #include "HAL/LowLevelMemTracker.h"
@@ -80,7 +81,7 @@
 			#include <objbase.h>
 		#include "HideWindowsPlatformTypes.h"
 	#endif
-#endif
+#endif //WITH_EDITOR
 
 #if WITH_ENGINE
 	#include "Engine/GameEngine.h"
@@ -95,7 +96,7 @@
 	#include "EngineGlobals.h"
 	#include "AudioThread.h"
 #if WITH_ENGINE && !UE_BUILD_SHIPPING
-	#include "Interfaces/IAutomationControllerModule.h"
+	#include "IAutomationControllerModule.h"
 #endif // WITH_ENGINE && !UE_BUILD_SHIPPING
 	#include "Database.h"
 	#include "DerivedDataCacheInterface.h"
@@ -128,7 +129,9 @@
 	#include "LongGPUTask.h"
 
 #if !UE_SERVER
+	#include "AppMediaTimeSource.h"
 	#include "IHeadMountedDisplayModule.h"
+	#include "IMediaModule.h"
 	#include "HeadMountedDisplay.h"
 	#include "MRMeshModule.h"
 	#include "Interfaces/ISlateRHIRendererModule.h"
@@ -147,9 +150,9 @@
 #endif
 
 #if WITH_AUTOMATION_WORKER
-	#include "Interfaces/IAutomationWorkerModule.h"
+	#include "IAutomationWorkerModule.h"
 #endif
-#endif  //WITH_ENGINE
+#endif //WITH_ENGINE
 
 class FSlateRenderer;
 class SViewport;
@@ -184,7 +187,7 @@ class FFeedbackContext;
 #endif
 
 #if defined(WITH_LAUNCHERCHECK) && WITH_LAUNCHERCHECK
-	#include "LauncherCheck.h"
+	#include "ILauncherCheckModule.h"
 #endif
 
 #if WITH_COREUOBJECT
@@ -2712,6 +2715,16 @@ int32 FEngineLoop::Init()
 
 	GetMoviePlayer()->WaitForMovieToFinish();
 
+#if !UE_SERVER
+	// initialize media framework
+	IMediaModule* MediaModule = FModuleManager::LoadModulePtr<IMediaModule>("Media");
+
+	if (MediaModule != nullptr)
+	{
+		MediaModule->SetTimeSource(MakeShareable(new FAppMediaTimeSource));
+	}
+#endif
+
 	// initialize automation worker
 #if WITH_AUTOMATION_WORKER
 	FModuleManager::Get().LoadModule("AutomationWorker");
@@ -3096,14 +3109,14 @@ void FEngineLoop::Tick()
 
 	SCOPED_NAMED_EVENT(FEngineLoopTick, FColor::Red);
 
-	// early in the Tick() to get the callbacks for cvar changes called
+	// execute callbacks for cvar changes
 	{
 		QUICK_SCOPE_CYCLE_COUNTER(STAT_FEngineLoop_Tick_CallAllConsoleVariableSinks);
 		IConsoleManager::Get().CallAllConsoleVariableSinks();
 	}
 
 	{
-		SCOPE_CYCLE_COUNTER( STAT_FrameTime );
+		SCOPE_CYCLE_COUNTER(STAT_FrameTime);
 
 		#if WITH_PROFILEGPU
 			// Issue the measurement of the execution time of a basic LongGPUTask unit on the very first frame
@@ -3120,6 +3133,7 @@ void FEngineLoop::Tick()
 			}
 		#endif
 
+		// beginning of RHI frame
 		ENQUEUE_UNIQUE_RENDER_COMMAND(
 			BeginFrame,
 		{
@@ -3140,35 +3154,36 @@ void FEngineLoop::Tick()
 
 		FCoreDelegates::OnBeginFrame.Broadcast();
 
+		// flush debug output which has been buffered by other threads
 		{
 			QUICK_SCOPE_CYCLE_COUNTER(STAT_FEngineLoop_FlushThreadedLogs);
-			// Flush debug output which has been buffered by other threads.
 			GLog->FlushThreadedLogs();
 		}
 
-		// Exit if frame limit is reached in benchmark mode.
-		if( (FApp::IsBenchmarking() && MaxFrameCounter && (GFrameCounter > MaxFrameCounter))
-			// or time limit is reached if set.
-			|| (MaxTickTime && (TotalTickTime > MaxTickTime)) )
+		// exit if frame limit is reached in benchmark mode, or if time limit is reached
+		if ((FApp::IsBenchmarking() && MaxFrameCounter && (GFrameCounter > MaxFrameCounter)) ||
+			(MaxTickTime && (TotalTickTime > MaxTickTime)))
 		{
 			FPlatformMisc::RequestExit(0);
 		}
 
+		// set FApp::CurrentTime, FApp::DeltaTime and potentially wait to enforce max tick rate
 		{
 			QUICK_SCOPE_CYCLE_COUNTER(STAT_FEngineLoop_UpdateTimeAndHandleMaxTickRate);
-			// Set FApp::CurrentTime, FApp::DeltaTime and potentially wait to enforce max tick rate.
 			GEngine->UpdateTimeAndHandleMaxTickRate();
 		}
 
+		// tick performance monitoring
 		{
 			QUICK_SCOPE_CYCLE_COUNTER(STAT_FEngineLoop_TickFPSChart);
 			GEngine->TickPerformanceMonitoring( FApp::GetDeltaTime() );
 		}
 
-		QUICK_SCOPE_CYCLE_COUNTER(STAT_FEngineLoop_Malloc_UpdateStats);
-
-		// Update memory allocator stats.
-		GMalloc->UpdateStats();
+		// update memory allocator stats
+		{
+			QUICK_SCOPE_CYCLE_COUNTER(STAT_FEngineLoop_Malloc_UpdateStats);
+			GMalloc->UpdateStats();
+		}
 	}
 
 	FStats::AdvanceFrame( false, FStats::FOnAdvanceRenderingThreadStats::CreateStatic( &AdvanceRenderingThreadStatsGT ) );
@@ -3222,7 +3237,7 @@ void FEngineLoop::Tick()
 			{
 				WorldToScale = GEditor->PlayWorld;
 			}
-#endif	// WITH_EDITOR
+#endif //WITH_EDITOR
 
 			if( WorldToScale != nullptr )
 			{
@@ -3234,11 +3249,12 @@ void FEngineLoop::Tick()
 
 			GNewWorldToMetersScale = 0.0f;
 		}
-#endif	// WITH_ENGINE
+#endif //WITH_ENGINE
 
-		// tick the currently active platform files
+		// tick active platform files
 		FPlatformFileManager::Get().TickActivePlatformFile();
 
+		// process accumulated Slate input
 		if (FSlateApplication::IsInitialized() && !bIdleMode)
 		{
 			SCOPE_TIME_GUARD(TEXT("SlateInput"));
@@ -3251,14 +3267,34 @@ void FEngineLoop::Tick()
 			SlateApp.FinishedInputThisFrame();
 		}
 
-		GEngine->Tick( FApp::GetDeltaTime(), bIdleMode );
+#if !UE_SERVER
+		// tick media framework
+		IMediaModule* MediaModule = FModuleManager::LoadModulePtr<IMediaModule>("Media");
+
+		if (MediaModule != nullptr)
+		{
+			QUICK_SCOPE_CYCLE_COUNTER(STAT_FEngineLoop_MediaTickPreEngine);
+			MediaModule->TickPreEngine();
+		}
+#endif
+
+		// main game engine tick (world, game objects, etc.)
+		GEngine->Tick(FApp::GetDeltaTime(), bIdleMode);
 
+#if !UE_SERVER
+		// tick media framework
+		if (MediaModule != nullptr)
+		{
+			QUICK_SCOPE_CYCLE_COUNTER(STAT_FEngineLoop_MediaTickPostEngine);
+			MediaModule->TickPostEngine();
+		}
+#endif
+		
 		// If a movie that is blocking the game thread has been playing,
 		// wait for it to finish before we continue to tick or tick again
 		// We do this right after GEngine->Tick() because that is where user code would initiate a load / movie.
 		{
 			QUICK_SCOPE_CYCLE_COUNTER(STAT_FEngineLoop_WaitForMovieToFinish);
-
 			GetMoviePlayer()->WaitForMovieToFinish();
 		}
 
@@ -3275,7 +3311,17 @@ void FEngineLoop::Tick()
 			GDistanceFieldAsyncQueue->ProcessAsyncTasks();
 		}
 
+#if !UE_SERVER
+		// tick media framework
+		if (MediaModule != nullptr)
+		{
+			QUICK_SCOPE_CYCLE_COUNTER(STAT_FEngineLoop_MediaTickPreSlate);
+			MediaModule->TickPreSlate();
+		}
+#endif
+
 #if WITH_ENGINE
+		// process concurrent Slate tasks
 		FGraphEventRef ConcurrentTask;
 		const bool bDoConcurrentSlateTick = GEngine->ShouldDoAsyncEndOfFrameTasks();
 
@@ -3289,22 +3335,22 @@ void FEngineLoop::Tick()
 			if (CurrentDemoNetDriver && CurrentDemoNetDriver->ShouldTickFlushAsyncEndOfFrame())
 			{
 				ConcurrentTask = TGraphTask<FExecuteConcurrentWithSlateTickTask>::CreateTask(nullptr, ENamedThreads::GameThread).ConstructAndDispatchWhenReady(
-					[CurrentDemoNetDriver, DeltaSeconds]()
-					{
+					[CurrentDemoNetDriver, DeltaSeconds]() {
 						if (CVarDoAsyncEndOfFrameTasksRandomize.GetValueOnAnyThread(true) > 0)
 						{
 							FPlatformProcess::Sleep(FMath::RandRange(0.0f, .003f)); // this shakes up the threading to find race conditions
 						}
+
 						if (CurrentDemoNetDriver != nullptr)
 						{
 							CurrentDemoNetDriver->TickFlushAsyncEndOfFrame(DeltaSeconds);
 						}
-					}
-				);
+					});
 			}
 		}
 #endif
 
+		// tick Slate application
 		if (FSlateApplication::IsInitialized() && !bIdleMode)
 		{
 			{
@@ -3315,7 +3361,6 @@ void FEngineLoop::Tick()
 				ProcessLocalPlayerSlateOperations();
 			}
 
-			// Tick Slate application
 			FSlateApplication::Get().Tick();
 		}
 
@@ -3329,28 +3374,27 @@ void FEngineLoop::Tick()
 #endif
 
 #if STATS
-		// Clear any stat group notifications we have pending just incase they weren't claimed during FSlateApplication::Get().Tick
+		// Clear any stat group notifications we have pending just in case they weren't claimed during FSlateApplication::Get().Tick
 		extern CORE_API void ClearPendingStatGroups();
 		ClearPendingStatGroups();
 #endif
 
 #if WITH_EDITOR && !UE_BUILD_SHIPPING
+		// tick automation controller (Editor only)
 		{
-			QUICK_SCOPE_CYCLE_COUNTER( STAT_FEngineLoop_Tick_AutomationController );
-			static FName AutomationController( "AutomationController" );
-			//Check if module loaded to support the change to allow this to be hot compilable.
-			if (FModuleManager::Get().IsModuleLoaded( AutomationController ))
+			QUICK_SCOPE_CYCLE_COUNTER(STAT_FEngineLoop_Tick_AutomationController);
+			static FName AutomationController("AutomationController");
+			if (FModuleManager::Get().IsModuleLoaded(AutomationController))
 			{
-				FModuleManager::GetModuleChecked<IAutomationControllerModule>( AutomationController ).Tick();
+				FModuleManager::GetModuleChecked<IAutomationControllerModule>(AutomationController).Tick();
 			}
 		}
 #endif
 
-#if WITH_ENGINE
-#if WITH_AUTOMATION_WORKER
+#if WITH_ENGINE && WITH_AUTOMATION_WORKER
+		// tick automation worker
 		{
 			QUICK_SCOPE_CYCLE_COUNTER(STAT_FEngineLoop_Tick_AutomationWorker);
-			//Check if module loaded to support the change to allow this to be hot compilable.
 			static const FName AutomationWorkerModuleName = TEXT("AutomationWorker");
 			if (FModuleManager::Get().IsModuleLoaded(AutomationWorkerModuleName))
 			{
@@ -3358,9 +3402,9 @@ void FEngineLoop::Tick()
 			}
 		}
 #endif
-#endif //WITH_ENGINE
 
-		{
+		// tick render hardware interface
+		{			
 			SCOPE_CYCLE_COUNTER(STAT_RHITickTime);
 			RHITick( FApp::GetDeltaTime() ); // Update RHI.
 		}
@@ -3369,46 +3413,52 @@ void FEngineLoop::Tick()
 		GFrameCounter++;
 
 		// Disregard first few ticks for total tick time as it includes loading and such.
-		if( GFrameCounter > 6 )
+		if (GFrameCounter > 6)
 		{
-			TotalTickTime+=FApp::GetDeltaTime();
+			TotalTickTime += FApp::GetDeltaTime();
 		}
 
-
 		// Find the objects which need to be cleaned up the next frame.
 		FPendingCleanupObjects* PreviousPendingCleanupObjects = PendingCleanupObjects;
 		PendingCleanupObjects = GetPendingCleanupObjects();
 
 		{
-			SCOPE_CYCLE_COUNTER( STAT_FrameSyncTime );
-			// this could be perhaps moved down to get greater parallelizm
+			SCOPE_CYCLE_COUNTER(STAT_FrameSyncTime);
+			// this could be perhaps moved down to get greater parallelism
 			// Sync game and render thread. Either total sync or allowing one frame lag.
 			static FFrameEndSync FrameEndSync;
 			static auto CVarAllowOneFrameThreadLag = IConsoleManager::Get().FindTConsoleVariableDataInt(TEXT("r.OneFrameThreadLag"));
 			FrameEndSync.Sync( CVarAllowOneFrameThreadLag->GetValueOnGameThread() != 0 );
 		}
 
+		// tick core ticker, threads & deferred commands
 		{
-			SCOPE_CYCLE_COUNTER( STAT_DeferredTickTime );
+			SCOPE_CYCLE_COUNTER(STAT_DeferredTickTime);
 			// Delete the objects which were enqueued for deferred cleanup before the previous frame.
 			delete PreviousPendingCleanupObjects;
 
-			// Destroy all linkers pending delete
 #if WITH_COREUOBJECT
-			DeleteLoaders();
+			DeleteLoaders(); // destroy all linkers pending delete
 #endif
 
 			FTicker::GetCoreTicker().Tick(FApp::GetDeltaTime());
-
 			FThreadManager::Get().Tick();
+			GEngine->TickDeferredCommands();		
+		}
 
-			GEngine->TickDeferredCommands();
+#if !UE_SERVER
+		// tick media framework
+		if (MediaModule != nullptr)
+		{
+			QUICK_SCOPE_CYCLE_COUNTER(STAT_FEngineLoop_MediaTickPostRender);
+			MediaModule->TickPostRender();
 		}
+#endif
 
 		FCoreDelegates::OnEndFrame.Broadcast();
 
-		ENQUEUE_UNIQUE_RENDER_COMMAND(
-			EndFrame,
+		// end of RHI frame
+		ENQUEUE_UNIQUE_RENDER_COMMAND(EndFrame,
 		{
 			RHICmdList.EndFrame();
 			GPU_STATS_ENDFRAME(RHICmdList);
@@ -3436,6 +3486,7 @@ void FEngineLoop::ClearPendingCleanupObjects()
 
 #endif // WITH_ENGINE
 
+
 static TAutoConsoleVariable<int32> CVarLogTimestamp(
 	TEXT("log.Timestamp"),
 	1,
@@ -3445,6 +3496,7 @@ static TAutoConsoleVariable<int32> CVarLogTimestamp(
 	TEXT("  2 = Log timestamps in seconds elapsed since GStartTime e.g. [0130.29][420]"),
 	ECVF_Default);
 
+
 static TAutoConsoleVariable<int32> CVarLogCategory(
 	TEXT("log.Category"),
 	1,
diff --git a/Engine/Source/Runtime/Launch/Public/LaunchEngineLoop.h b/Engine/Source/Runtime/Launch/Public/LaunchEngineLoop.h
--- a/Engine/Source/Runtime/Launch/Public/LaunchEngineLoop.h
+++ b/Engine/Source/Runtime/Launch/Public/LaunchEngineLoop.h
@@ -12,6 +12,7 @@ class FEngineService;
 class FPendingCleanupObjects;
 class ISessionService;
 
+
 /**
  * Implements the main engine loop.	
  */
@@ -22,9 +23,7 @@ class FEngineLoop
 {
 public:
 
-	/**
-	 * Default constructor.
-	 */
+	/** Default constructor. */
 	FEngineLoop();
 
 	virtual ~FEngineLoop() { }
@@ -39,7 +38,7 @@ public:
 	 * @param AdditionalCommandLine Optional string to append to the command line (after ArgV is put together).
 	 * @return Returns the error level, 0 if successful and > 0 if there were errors.
 	 */ 
-	int32 PreInit(int32 ArgC, TCHAR* ArgV[], const TCHAR* AdditionalCommandline = NULL);
+	int32 PreInit(int32 ArgC, TCHAR* ArgV[], const TCHAR* AdditionalCommandline = nullptr);
 
 	/**
 	 * Pre-Initialize the main loop - parse command line, sets up GIsEditor, etc.
@@ -93,21 +92,21 @@ public:
 public:
 
 	/** Initializes the application. */
-	static bool AppInit( );
+	static bool AppInit();
 
 	/**
 	 * Prepares the application for shutdown.
 	 *
 	 * This function is called from within guarded exit code, only during non-error exits.
 	 */
-	static void AppPreExit( );
+	static void AppPreExit();
 
 	/**
 	 * Shuts down the application.
 	 *
 	 * This function called outside guarded exit code, during all exits (including error exits).
 	 */
-	static void AppExit( );
+	static void AppExit();
 
 private:
 
@@ -152,7 +151,5 @@ private:
 };
 
 
-/**
- * Global engine loop object. This is needed so wxWindows can access it.
- */
+/** Global engine loop object. This is needed so wxWindows can access it. */
 extern FEngineLoop GEngineLoop;
