Fix for TTY on consoles being doubled-up and causing issues with Orion tests

diff --git a/Engine/Source/Runtime/Launch/Private/LaunchEngineLoop.cpp b/Engine/Source/Runtime/Launch/Private/LaunchEngineLoop.cpp
--- a/Engine/Source/Runtime/Launch/Private/LaunchEngineLoop.cpp
+++ b/Engine/Source/Runtime/Launch/Private/LaunchEngineLoop.cpp
@@ -978,10 +978,14 @@ int32 FEngineLoop::PreInit(const TCHAR* CmdLine)
 	GLog->EnableBacklog(true);
 
 	// Initialize std out device as early as possible if requested in the command line
+#if PLATFORM_DESKTOP
+	// consoles don't typically have stdout, and FOutputDeviceDebug is responsible for echoing logs to the
+	// terminal
 	if (FParse::Param(FCommandLine::Get(), TEXT("stdout")))
 	{
 		InitializeStdOutDevice();
 	}
+#endif
 
 #if !UE_BUILD_SHIPPING
 	if (FPlatformProperties::SupportsQuit())
