Fix crash on exit when TaskGraph is not run (UE-56510).

(Edigrating CL 3967173 to Main)

diff --git a/Engine/Source/Runtime/Launch/Private/LaunchEngineLoop.cpp b/Engine/Source/Runtime/Launch/Private/LaunchEngineLoop.cpp
--- a/Engine/Source/Runtime/Launch/Private/LaunchEngineLoop.cpp
+++ b/Engine/Source/Runtime/Launch/Private/LaunchEngineLoop.cpp
@@ -375,7 +375,7 @@ static void RHIExitAndStopRHIThread()
 #endif
 
 	// Stop the RHI Thread (using GRHIThread_InternalUseOnly is unreliable since RT may be stopped)
-	if (FTaskGraphInterface::Get().IsThreadProcessingTasks(ENamedThreads::RHIThread))
+	if (FTaskGraphInterface::IsRunning() && FTaskGraphInterface::Get().IsThreadProcessingTasks(ENamedThreads::RHIThread))
 	{
 		DECLARE_CYCLE_STAT(TEXT("Wait For RHIThread Finish"), STAT_WaitForRHIThreadFinish, STATGROUP_TaskGraphTasks);
 		FGraphEventRef QuitTask = TGraphTask<FReturnGraphTask>::CreateTask(nullptr, ENamedThreads::GameThread).ConstructAndDispatchWhenReady(ENamedThreads::RHIThread);
