Fix assumption that all string literals use wide characters in Version.h

diff --git a/Engine/Source/Runtime/Launch/Resources/Version.h b/Engine/Source/Runtime/Launch/Resources/Version.h
--- a/Engine/Source/Runtime/Launch/Resources/Version.h
+++ b/Engine/Source/Runtime/Launch/Resources/Version.h
@@ -59,8 +59,7 @@
 #define ENGINE_PATCH_VERSION	0
 
 // Macros for encoding strings
-#define VERSION_TEXT_2(x) L ## x
-#define VERSION_TEXT(x) VERSION_TEXT_2(x)
+#define VERSION_TEXT(x) TEXT(x)
 #define VERSION_STRINGIFY_2(x) VERSION_TEXT(#x)
 #define VERSION_STRINGIFY(x) VERSION_STRINGIFY_2(x)
 
