Updated minumum macOS version required to run the editor to 10.13.5

diff --git a/Engine/Source/Editor/UnrealEd/Private/EditorEngine.cpp b/Engine/Source/Editor/UnrealEd/Private/EditorEngine.cpp
--- a/Engine/Source/Editor/UnrealEd/Private/EditorEngine.cpp
+++ b/Engine/Source/Editor/UnrealEd/Private/EditorEngine.cpp
@@ -6711,7 +6711,7 @@ void UEditorEngine::UpdateAutoLoadProject()
 #if PLATFORM_MAC
 	if ( !GIsBuildMachine )
 	{
-		if(FPlatformMisc::MacOSXVersionCompare(10,12,5) < 0)
+		if(FPlatformMisc::MacOSXVersionCompare(10,13,5) < 0)
 		{
 			if(FSlateApplication::IsInitialized())
 			{
diff --git a/Engine/Source/Runtime/Launch/Private/Mac/LaunchMac.cpp b/Engine/Source/Runtime/Launch/Private/Mac/LaunchMac.cpp
--- a/Engine/Source/Runtime/Launch/Private/Mac/LaunchMac.cpp
+++ b/Engine/Source/Runtime/Launch/Private/Mac/LaunchMac.cpp
@@ -254,17 +254,23 @@ static int32 MacOSVersionCompare(const NSOperatingSystemVersion& VersionA, const
 	NSString* MinimumSystemVersionString = (NSString*)InfoDictionary[@"LSMinimumSystemVersion"];
 	NSOperatingSystemVersion MinimumSystemVersion = { 0 };
 	NSOperatingSystemVersion CurrentSystemVersion = [[NSProcessInfo processInfo] operatingSystemVersion];
-	NSOperatingSystemVersion LatestSierraSystemVersion = { 10, 12, 6 };
+#if WITH_EDITOR
+	NSOperatingSystemVersion LatestMacOSVersion = { 10, 13, 5 };
+	NSString* LatestMacOSVersionString = @"10.13.5";
+#else
+	NSOperatingSystemVersion LatestMacOSVersion = { 10, 12, 6 };
+	NSString* LatestMacOSVersionString = @"10.12.6";
+#endif
 	NSArray<NSString*>* VersionComponents = [MinimumSystemVersionString componentsSeparatedByString:@"."];
 	MinimumSystemVersion.majorVersion = [[VersionComponents objectAtIndex:0] integerValue];
 	MinimumSystemVersion.minorVersion = VersionComponents.count > 1 ? [[VersionComponents objectAtIndex:1] integerValue] : 0;
 	MinimumSystemVersion.patchVersion = VersionComponents.count > 2 ? [[VersionComponents objectAtIndex:2] integerValue] : 0;
 
-	// Make sure that the min version in Info.plist is at least 10.12.6, as that's the absolute minimum
-	if (MacOSVersionCompare(MinimumSystemVersion, LatestSierraSystemVersion) < 0)
+	// Make sure that the min version in Info.plist is at least 10.12.6 for games or 10.13.5 for the editor, as that's the absolute minimum
+	if (MacOSVersionCompare(MinimumSystemVersion, LatestMacOSVersion) < 0)
 	{
-		MinimumSystemVersion = LatestSierraSystemVersion;
-		MinimumSystemVersionString = @"10.12.6";
+		MinimumSystemVersion = LatestMacOSVersion;
+		MinimumSystemVersionString = LatestMacOSVersionString;
 	}
 
 	if (MacOSVersionCompare(CurrentSystemVersion, MinimumSystemVersion) < 0)
