UE-57136 Assertion failed: InWindow when starting Android app with phone awake but screen locked on Galaxy S9

#Adnroid
#4.20

Wait for on the main loop, before PreInit, for re-creating the native window, if previously destroyed by onStop()

diff --git a/Engine/Source/Runtime/Launch/Private/Android/LaunchAndroid.cpp b/Engine/Source/Runtime/Launch/Private/Android/LaunchAndroid.cpp
--- a/Engine/Source/Runtime/Launch/Private/Android/LaunchAndroid.cpp
+++ b/Engine/Source/Runtime/Launch/Private/Android/LaunchAndroid.cpp
@@ -371,6 +371,13 @@ int32 AndroidMain(struct android_app* state)
 	// OBBs and APK are found.
 	IPlatformFile::GetPlatformPhysical().Initialize(nullptr, FCommandLine::Get());
 
+	//wait for re-creating the native window, if previously destroyed by onStop()
+	while (FAndroidWindow::GetHardwareWindow() == nullptr)
+	{
+		FPlatformProcess::Sleep(0.01f);
+		FPlatformMisc::MemoryBarrier();
+	}
+
 	// initialize the engine
 	int32 PreInitResult = GEngineLoop.PreInit(0, NULL, FCommandLine::Get());
 
