Allow the "compile on startup" functionality to run for content-only projects. It's not uncommon for QA to build one project and switch to another.

diff --git a/Engine/Source/Runtime/Launch/Private/LaunchEngineLoop.cpp b/Engine/Source/Runtime/Launch/Private/LaunchEngineLoop.cpp
--- a/Engine/Source/Runtime/Launch/Private/LaunchEngineLoop.cpp
+++ b/Engine/Source/Runtime/Launch/Private/LaunchEngineLoop.cpp
@@ -3956,86 +3956,82 @@ bool FEngineLoop::AppInit( )
 #if UE_EDITOR && !IS_MONOLITHIC
 	if(!GIsBuildMachine && FPaths::IsProjectFilePathSet() && PluginManager.AreRequiredPluginsAvailable())
 	{
-		const FProjectDescriptor* CurrentProject = ProjectManager.GetCurrentProject();
-		if(CurrentProject != nullptr && CurrentProject->Modules.Num() > 0)
+		bool bNeedCompile = false;
+		GConfig->GetBool(TEXT("/Script/UnrealEd.EditorLoadingSavingSettings"), TEXT("bForceCompilationAtStartup"), bNeedCompile, GEditorPerProjectIni);
+		if(FParse::Param(FCommandLine::Get(), TEXT("SKIPCOMPILE")) || FParse::Param(FCommandLine::Get(), TEXT("MULTIPROCESS")))
 		{
-			bool bNeedCompile = false;
-			GConfig->GetBool(TEXT("/Script/UnrealEd.EditorLoadingSavingSettings"), TEXT("bForceCompilationAtStartup"), bNeedCompile, GEditorPerProjectIni);
-			if(FParse::Param(FCommandLine::Get(), TEXT("SKIPCOMPILE")) || FParse::Param(FCommandLine::Get(), TEXT("MULTIPROCESS")))
-			{
-				bNeedCompile = false;
-			}
-			if(!bNeedCompile)
-			{
-				// Check if any of the project or plugin modules are out of date, and the user wants to compile them.
-				TArray<FString> IncompatibleFiles;
-				ProjectManager.CheckModuleCompatibility(IncompatibleFiles);
-				PluginManager.CheckModuleCompatibility(IncompatibleFiles);
+			bNeedCompile = false;
+		}
+		if(!bNeedCompile)
+		{
+			// Check if any of the project or plugin modules are out of date, and the user wants to compile them.
+			TArray<FString> IncompatibleFiles;
+			ProjectManager.CheckModuleCompatibility(IncompatibleFiles);
+			PluginManager.CheckModuleCompatibility(IncompatibleFiles);
 
-				if (IncompatibleFiles.Num() > 0)
+			if (IncompatibleFiles.Num() > 0)
+			{
+				// Log the modules which need to be rebuilt
+				for (int Idx = 0; Idx < IncompatibleFiles.Num(); Idx++)
 				{
-					// Log the modules which need to be rebuilt
-					for (int Idx = 0; Idx < IncompatibleFiles.Num(); Idx++)
-					{
-						UE_LOG(LogInit, Warning, TEXT("Incompatible or missing module: %s"), *IncompatibleFiles[Idx]);
-					}
+					UE_LOG(LogInit, Warning, TEXT("Incompatible or missing module: %s"), *IncompatibleFiles[Idx]);
+				}
 
-					// Build the error message for the dialog box
-					FString ModulesList = TEXT("The following modules are missing or built with a different engine version:\n\n");
+				// Build the error message for the dialog box
+				FString ModulesList = TEXT("The following modules are missing or built with a different engine version:\n\n");
 
-					int NumModulesToDisplay = (IncompatibleFiles.Num() <= 20)? IncompatibleFiles.Num() : 15;
-					for (int Idx = 0; Idx < NumModulesToDisplay; Idx++)
-					{
-						ModulesList += FString::Printf(TEXT("  %s\n"), *IncompatibleFiles[Idx]);
-					}
-					if(IncompatibleFiles.Num() > NumModulesToDisplay)
-					{
-						ModulesList += FString::Printf(TEXT("  (+%d others, see log for details)\n"), IncompatibleFiles.Num() - NumModulesToDisplay);
-					}
-
-					ModulesList += TEXT("\nWould you like to rebuild them now?");
+				int NumModulesToDisplay = (IncompatibleFiles.Num() <= 20)? IncompatibleFiles.Num() : 15;
+				for (int Idx = 0; Idx < NumModulesToDisplay; Idx++)
+				{
+					ModulesList += FString::Printf(TEXT("  %s\n"), *IncompatibleFiles[Idx]);
+				}
+				if(IncompatibleFiles.Num() > NumModulesToDisplay)
+				{
+					ModulesList += FString::Printf(TEXT("  (+%d others, see log for details)\n"), IncompatibleFiles.Num() - NumModulesToDisplay);
+				}
 
-					// If we're running with -stdout, assume that we're a non interactive process and about to fail
-					if (FApp::IsUnattended() || FParse::Param(FCommandLine::Get(), TEXT("stdout")))
-					{
-						return false;
-					}
+				ModulesList += TEXT("\nWould you like to rebuild them now?");
 
-					// Ask whether to compile before continuing
-					if (FPlatformMisc::MessageBoxExt(EAppMsgType::YesNo, *ModulesList, *FString::Printf(TEXT("Missing %s Modules"), FApp::GetProjectName())) == EAppReturnType::No)
-					{
-						return false;
-					}
+				// If we're running with -stdout, assume that we're a non interactive process and about to fail
+				if (FApp::IsUnattended() || FParse::Param(FCommandLine::Get(), TEXT("stdout")))
+				{
+					return false;
+				}
 
-					bNeedCompile = true;
+				// Ask whether to compile before continuing
+				if (FPlatformMisc::MessageBoxExt(EAppMsgType::YesNo, *ModulesList, *FString::Printf(TEXT("Missing %s Modules"), FApp::GetProjectName())) == EAppReturnType::No)
+				{
+					return false;
 				}
+
+				bNeedCompile = true;
 			}
+		}
 
-			if(bNeedCompile)
+		if(bNeedCompile)
+		{
+			// Try to compile it
+			FFeedbackContext *Context = (FFeedbackContext*)FDesktopPlatformModule::Get()->GetNativeFeedbackContext();
+			Context->BeginSlowTask(FText::FromString(TEXT("Starting build...")), true, true);
+			bool bCompileResult = FDesktopPlatformModule::Get()->CompileGameProject(FPaths::RootDir(), FPaths::GetProjectFilePath(), Context);
+			Context->EndSlowTask();
+
+			// Get a list of modules which are still incompatible
+			TArray<FString> StillIncompatibleFiles;
+			ProjectManager.CheckModuleCompatibility(StillIncompatibleFiles);
+			PluginManager.CheckModuleCompatibility(StillIncompatibleFiles);
+
+			if(!bCompileResult || StillIncompatibleFiles.Num() > 0)
 			{
-				// Try to compile it
-				FFeedbackContext *Context = (FFeedbackContext*)FDesktopPlatformModule::Get()->GetNativeFeedbackContext();
-				Context->BeginSlowTask(FText::FromString(TEXT("Starting build...")), true, true);
-				bool bCompileResult = FDesktopPlatformModule::Get()->CompileGameProject(FPaths::RootDir(), FPaths::GetProjectFilePath(), Context);
-				Context->EndSlowTask();
-
-				// Get a list of modules which are still incompatible
-				TArray<FString> StillIncompatibleFiles;
-				ProjectManager.CheckModuleCompatibility(StillIncompatibleFiles);
-				PluginManager.CheckModuleCompatibility(StillIncompatibleFiles);
-
-				if(!bCompileResult || StillIncompatibleFiles.Num() > 0)
+				for (int Idx = 0; Idx < StillIncompatibleFiles.Num(); Idx++)
 				{
-					for (int Idx = 0; Idx < StillIncompatibleFiles.Num(); Idx++)
-					{
-						UE_LOG(LogInit, Warning, TEXT("Still incompatible or missing module: %s"), *StillIncompatibleFiles[Idx]);
-					}
-					if (!FApp::IsUnattended())
-					{
-						FPlatformMisc::MessageBoxExt(EAppMsgType::Ok, *FString::Printf(TEXT("%s could not be compiled. Try rebuilding from source manually."), FApp::GetProjectName()), TEXT("Error"));
-					}
-					return false;
+					UE_LOG(LogInit, Warning, TEXT("Still incompatible or missing module: %s"), *StillIncompatibleFiles[Idx]);
+				}
+				if (!FApp::IsUnattended())
+				{
+					FPlatformMisc::MessageBoxExt(EAppMsgType::Ok, *FString::Printf(TEXT("%s could not be compiled. Try rebuilding from source manually."), FApp::GetProjectName()), TEXT("Error"));
 				}
+				return false;
 			}
 		}
 	}
diff --git a/Engine/Source/Runtime/Projects/Private/PluginManager.cpp b/Engine/Source/Runtime/Projects/Private/PluginManager.cpp
--- a/Engine/Source/Runtime/Projects/Private/PluginManager.cpp
+++ b/Engine/Source/Runtime/Projects/Private/PluginManager.cpp
@@ -760,22 +760,6 @@ bool FPluginManager::ConfigureEnabledPlugin(const FPluginReferenceDescriptor& Fi
 				// Mount the binaries directory
 				const FString PluginBinariesPath = FPaths::Combine(*FPaths::GetPath(Plugin.FileName), TEXT("Binaries"), FPlatformProcess::GetBinariesSubdirectory());
 				FModuleManager::Get().AddBinariesDirectory(*PluginBinariesPath, Plugin.GetLoadedFrom() == EPluginLoadedFrom::Project);
-
-				// If this is a content-only project, make sure the modules are compatible with the engine (or allow the user to disable it). If it's a code project,
-				// we'll run a separate check in LaunchEngineLoop.cpp to do the compile-on-startup stuff.
-				const FProjectDescriptor* Project = IProjectManager::Get().GetCurrentProject();
-				if (Project != nullptr && Project->Modules.Num() == 0)
-				{
-					TArray<FString> IncompatibleFiles;
-					if (!FModuleDescriptor::CheckModuleCompatibility(Plugin.Descriptor.Modules, IncompatibleFiles))
-					{
-						if (PromptToDisableIncompatiblePlugin(FirstReference.Name, Reference.Name))
-						{
-							UE_LOG(LogPluginManager, Display, TEXT("Disabled plugin '%s', continuing."), *FirstReference.Name);
-							return true;
-						}
-					}
-				}
 			}
 
 			// Check the declared engine version. This is a soft requirement, so allow the user to skip over it.
