Allow some hardware buttons through if hardware keyboard disabled

#ue4
#android

diff --git a/Engine/Source/Runtime/Launch/Private/Android/LaunchAndroid.cpp b/Engine/Source/Runtime/Launch/Private/Android/LaunchAndroid.cpp
--- a/Engine/Source/Runtime/Launch/Private/Android/LaunchAndroid.cpp
+++ b/Engine/Source/Runtime/Launch/Private/Android/LaunchAndroid.cpp
@@ -67,6 +67,18 @@ static const int32_t AxisList[] =
 // map of all supported keycodes
 static TSet<uint32> MappedKeyCodes;
 
+// map of always allowed keycodes
+static TSet<uint32> AlwaysAllowedKeyCodes;
+
+// List of always allowed keycodes
+static const uint32 AlwaysAllowedKeyCodesList[] =
+{
+	AKEYCODE_MENU,
+	AKEYCODE_BACK,
+	AKEYCODE_VOLUME_UP,
+	AKEYCODE_VOLUME_DOWN
+};
+
 // List of desired gamepad keycodes
 static const uint32 ValidGamepadKeyCodesList[] =
 {
@@ -342,6 +354,12 @@ int32 AndroidMain(struct android_app* state)
 		MappedKeyCodes.Add(KeyCodes[i]);
 	}
 
+	const int AlwaysAllowedKeyCodesCount = sizeof(AlwaysAllowedKeyCodesList) / sizeof(uint32);
+	for (int i = 0; i < AlwaysAllowedKeyCodesCount; ++i)
+	{
+		AlwaysAllowedKeyCodes.Add(AlwaysAllowedKeyCodesList[i]);
+	}
+
 	const int ValidGamepadKeyCodeCount = sizeof(ValidGamepadKeyCodesList)/sizeof(uint32);
 	for (int i = 0; i < ValidGamepadKeyCodeCount; ++i)
 	{
@@ -918,7 +936,7 @@ static int32_t HandleInputCB(struct android_app* app, AInputEvent* event)
 				return 0;
 			}
 
-			if (bSoftKey || GAndroidEnableHardwareKeyboard)
+			if (bSoftKey || GAndroidEnableHardwareKeyboard || AlwaysAllowedKeyCodes.Contains(keyCode))
 			{
 				FDeferredAndroidMessage Message;
 
