Fix virtual keyboard crashes on Android

diff --git a/Engine/Source/Runtime/Launch/Private/Android/AndroidJNI.cpp b/Engine/Source/Runtime/Launch/Private/Android/AndroidJNI.cpp
--- a/Engine/Source/Runtime/Launch/Private/Android/AndroidJNI.cpp
+++ b/Engine/Source/Runtime/Launch/Private/Android/AndroidJNI.cpp
@@ -25,7 +25,7 @@ THIRD_PARTY_INCLUDES_END
 JavaVM* GJavaVM;
 
 // Pointer to target widget for virtual keyboard contents
-static IVirtualKeyboardEntry *VirtualKeyboardWidget = NULL;
+static TWeakPtr<IVirtualKeyboardEntry, ESPMode::ThreadSafe> VirtualKeyboardWidget;
 
 //virtualKeyboard shown
 static volatile bool GVirtualKeyboardShown = false;
@@ -251,7 +251,7 @@ int32 FJavaWrapper::CallIntMethod(JNIEnv* Env, jobject Object, jmethodID Method,
 {
 	if (Method == NULL || Object == NULL)
 	{
-		return false;
+		return 0;
 	}
 
 	va_list Args;
@@ -262,6 +262,51 @@ int32 FJavaWrapper::CallIntMethod(JNIEnv* Env, jobject Object, jmethodID Method,
 	return (int32)Return;
 }
 
+int64 FJavaWrapper::CallLongMethod(JNIEnv* Env, jobject Object, jmethodID Method, ...)
+{
+	if (Method == NULL || Object == NULL)
+	{
+		return 0;
+	}
+
+	va_list Args;
+	va_start(Args, Method);
+	jlong Return = Env->CallLongMethodV(Object, Method, Args);
+	va_end(Args);
+
+	return (int64)Return;
+}
+
+float FJavaWrapper::CallFloatMethod(JNIEnv* Env, jobject Object, jmethodID Method, ...)
+{
+	if (Method == NULL || Object == NULL)
+	{
+		return 0.0f;
+	}
+
+	va_list Args;
+	va_start(Args, Method);
+	jfloat Return = Env->CallFloatMethodV(Object, Method, Args);
+	va_end(Args);
+
+	return (float)Return;
+}
+
+double FJavaWrapper::CallDoubleMethod(JNIEnv* Env, jobject Object, jmethodID Method, ...)
+{
+	if (Method == NULL || Object == NULL)
+	{
+		return 0.0;
+	}
+
+	va_list Args;
+	va_start(Args, Method);
+	jdouble Return = Env->CallDoubleMethodV(Object, Method, Args);
+	va_end(Args);
+
+	return (double)Return;
+}
+
 bool FJavaWrapper::CallBooleanMethod(JNIEnv* Env, jobject Object, jmethodID Method, ...)
 {
 	if (Method == NULL || Object == NULL)
@@ -437,30 +482,32 @@ bool AndroidThunkCpp_GetInputDeviceInfo(int32 deviceId, FAndroidInputDeviceInfo
 	if (JNIEnv* Env = FAndroidApplication::GetJavaEnv())
 	{
 		jobject deviceInfo = (jobject)Env->CallObjectMethod(FJavaWrapper::GameActivityThis, FJavaWrapper::AndroidThunkJava_GetInputDeviceInfo, deviceId);
-		bool bIsOptional = false;
-		results.DeviceId = deviceId;
-		results.VendorId = (int32)Env->GetIntField(deviceInfo, FJavaWrapper::InputDeviceInfo_VendorId);
-		results.ProductId = (int32)Env->GetIntField(deviceInfo, FJavaWrapper::InputDeviceInfo_ProductId);
-		results.ControllerId = (int32)Env->GetIntField(deviceInfo, FJavaWrapper::InputDeviceInfo_ControllerId);
-
-		jstring jsName = (jstring)Env->GetObjectField(deviceInfo, FJavaWrapper::InputDeviceInfo_Name);
-		CHECK_JNI_RESULT(jsName);
-		const char * nativeName = Env->GetStringUTFChars(jsName, 0);
-		results.Name = FString(nativeName);
-		Env->ReleaseStringUTFChars(jsName, nativeName);
-		Env->DeleteLocalRef(jsName);
-
-		jstring jsDescriptor = (jstring)Env->GetObjectField(deviceInfo, FJavaWrapper::InputDeviceInfo_Descriptor);
-		CHECK_JNI_RESULT(jsDescriptor);
-		const char * nativeDescriptor = Env->GetStringUTFChars(jsDescriptor, 0);
-		results.Descriptor = FString(nativeDescriptor);
-		Env->ReleaseStringUTFChars(jsDescriptor, nativeDescriptor);
-		Env->DeleteLocalRef(jsDescriptor);
-		
-		// release references
-		Env->DeleteLocalRef(deviceInfo);
-
-		return true;
+		if (deviceInfo != 0)
+		{
+			bool bIsOptional = false;
+			results.DeviceId = deviceId;
+			results.VendorId = (int32)Env->GetIntField(deviceInfo, FJavaWrapper::InputDeviceInfo_VendorId);
+			results.ProductId = (int32)Env->GetIntField(deviceInfo, FJavaWrapper::InputDeviceInfo_ProductId);
+			results.ControllerId = (int32)Env->GetIntField(deviceInfo, FJavaWrapper::InputDeviceInfo_ControllerId);
+
+			jstring jsName = (jstring)Env->GetObjectField(deviceInfo, FJavaWrapper::InputDeviceInfo_Name);
+			CHECK_JNI_RESULT(jsName);
+			const char * nativeName = Env->GetStringUTFChars(jsName, 0);
+			results.Name = FString(nativeName);
+			Env->ReleaseStringUTFChars(jsName, nativeName);
+			Env->DeleteLocalRef(jsName);
+
+			jstring jsDescriptor = (jstring)Env->GetObjectField(deviceInfo, FJavaWrapper::InputDeviceInfo_Descriptor);
+			CHECK_JNI_RESULT(jsDescriptor);
+			const char * nativeDescriptor = Env->GetStringUTFChars(jsDescriptor, 0);
+			results.Descriptor = FString(nativeDescriptor);
+			Env->ReleaseStringUTFChars(jsDescriptor, nativeDescriptor);
+			Env->DeleteLocalRef(jsDescriptor);
+
+			// release references
+			Env->DeleteLocalRef(deviceInfo);
+			return true;
+		}
 	}
 
 	// failed
@@ -499,7 +546,7 @@ JNI_METHOD void Java_com_epicgames_ue4_GameActivity_nativeVirtualKeyboardVisible
 	//remove reference so the object can be clicked again to show the virtual keyboard
 	if (!bShown)
 	{
-		VirtualKeyboardWidget = NULL;
+		VirtualKeyboardWidget.Reset();
 	}
 }
 
@@ -732,30 +779,30 @@ void AndroidThunkCpp_ShowVirtualKeyboardInputDialog(TSharedPtr<IVirtualKeyboardE
     if (JNIEnv* Env = FAndroidApplication::GetJavaEnv())
     {
         // remember target widget for contents
-        VirtualKeyboardWidget = &(*TextWidget);
-        
-        // call the java side
-        jstring LabelJava = Env->NewStringUTF(TCHAR_TO_UTF8(*Label));
-        jstring ContentsJava = Env->NewStringUTF(TCHAR_TO_UTF8(*Contents));
-        FJavaWrapper::CallVoidMethod(Env, FJavaWrapper::GameActivityThis, FJavaWrapper::AndroidThunkJava_ShowVirtualKeyboardInputDialog, InputType, LabelJava, ContentsJava);
-        Env->DeleteLocalRef(ContentsJava);
-        Env->DeleteLocalRef(LabelJava);
+		VirtualKeyboardWidget = TextWidget;
+
+		// call the java side
+		jstring LabelJava = Env->NewStringUTF(TCHAR_TO_UTF8(*Label));
+		jstring ContentsJava = Env->NewStringUTF(TCHAR_TO_UTF8(*Contents));
+		FJavaWrapper::CallVoidMethod(Env, FJavaWrapper::GameActivityThis, FJavaWrapper::AndroidThunkJava_ShowVirtualKeyboardInputDialog, InputType, LabelJava, ContentsJava);
+		Env->DeleteLocalRef(ContentsJava);
+		Env->DeleteLocalRef(LabelJava);
     }
 }
 
 void AndroidThunkCpp_HideVirtualKeyboardInputDialog()
 {
     // Make sure virtual keyboard currently open
-    if (VirtualKeyboardWidget == NULL)
-    {
-        return;
-    }
+	if (!VirtualKeyboardWidget.IsValid())
+	{
+		return;
+	}
     
     if (JNIEnv* Env = FAndroidApplication::GetJavaEnv())
     {
         // ignore anything it might return
-        VirtualKeyboardWidget = NULL;
-        
+		VirtualKeyboardWidget.Reset();
+
         // call the java side
         FJavaWrapper::CallVoidMethod(Env, FJavaWrapper::GameActivityThis, FJavaWrapper::AndroidThunkJava_HideVirtualKeyboardInputDialog);
 
@@ -786,7 +833,7 @@ JNI_METHOD void Java_com_epicgames_ue4_GameActivity_nativeVirtualKeyboardShown(J
 void AndroidThunkCpp_HideVirtualKeyboardInput()
 {
 	// Make sure virtual keyboard currently open
-	if (VirtualKeyboardWidget == NULL)
+	if (!VirtualKeyboardWidget.IsValid())
 	{
 		return;
 	}
@@ -794,7 +841,7 @@ void AndroidThunkCpp_HideVirtualKeyboardInput()
 	if (JNIEnv* Env = FAndroidApplication::GetJavaEnv())
 	{
 		// ignore anything it might return
-		VirtualKeyboardWidget = NULL;
+		VirtualKeyboardWidget.Reset();
 
 		// call the java side
 		FJavaWrapper::CallVoidMethod(Env, FJavaWrapper::GameActivityThis, FJavaWrapper::AndroidThunkJava_HideVirtualKeyboardInput);
@@ -811,18 +858,16 @@ void AndroidThunkCpp_HideVirtualKeyboardInput()
 
 void AndroidThunkCpp_ShowVirtualKeyboardInput(TSharedPtr<IVirtualKeyboardEntry> TextWidget, int32 InputType, const FString& Label, const FString& Contents)
 {
-	// remember target widget for contents
-	IVirtualKeyboardEntry * newWidget = &(*TextWidget);
 	//#jira UE-49139 Tapping in the same text box doesn't make the virtual keyboard disappear
-	if (VirtualKeyboardWidget == newWidget)
+	if (VirtualKeyboardWidget.HasSameObject(TextWidget.Get()))
 	{
 		FPlatformMisc::LowLevelOutputDebugString(TEXT("[JNI] - AndroidThunkCpp_ShowVirtualKeyboardInput same control"));
 		AndroidThunkCpp_HideVirtualKeyboardInput();
 	}
 	else if (JNIEnv* Env = FAndroidApplication::GetJavaEnv())
 	{
-
-		VirtualKeyboardWidget = newWidget;
+		// remember target widget for contents
+		VirtualKeyboardWidget = TextWidget;
 
 		// call the java side
 		jstring LabelJava = Env->NewStringUTF(TCHAR_TO_UTF8(*Label));
@@ -839,7 +884,7 @@ JNI_METHOD void Java_com_epicgames_ue4_GameActivity_nativeVirtualKeyboardResult(
 	// update text widget with new contents if OK pressed
 	if (update == JNI_TRUE)
 	{
-		if (VirtualKeyboardWidget != NULL)
+		if (VirtualKeyboardWidget.IsValid())
 		{
 			const char* javaChars = jenv->GetStringUTFChars(contents, 0);
 
@@ -848,26 +893,37 @@ JNI_METHOD void Java_com_epicgames_ue4_GameActivity_nativeVirtualKeyboardResult(
 			{
 				FGraphEventRef SetWidgetText = FFunctionGraphTask::CreateAndDispatchWhenReady([&]()
 				{
-					if (VirtualKeyboardWidget != NULL)
+					TSharedPtr<IVirtualKeyboardEntry> LockedKeyboardWidget(VirtualKeyboardWidget.Pin());
+					if (LockedKeyboardWidget.IsValid())
 					{
-						VirtualKeyboardWidget->SetTextFromVirtualKeyboard(FText::FromString(FString(UTF8_TO_TCHAR(javaChars))), ETextEntryType::TextEntryAccepted);
+						LockedKeyboardWidget->SetTextFromVirtualKeyboard(FText::FromString(FString(UTF8_TO_TCHAR(javaChars))), ETextEntryType::TextEntryAccepted);
 					}
+
+					// release reference
+					VirtualKeyboardWidget.Reset();
 				}, TStatId(), NULL, ENamedThreads::GameThread);
 				FTaskGraphInterface::Get().WaitUntilTaskCompletes(SetWidgetText);
 			}
+			else
+			{
+				// release reference
+				VirtualKeyboardWidget.Reset();
+			}
 			// release string
 			jenv->ReleaseStringUTFChars(contents, javaChars);
 		}
 	}
-
-	// release reference
-	VirtualKeyboardWidget = NULL;
+	else
+	{
+		// release reference
+		VirtualKeyboardWidget.Reset();
+	}
 }
 
 //This function is declared in the Java-defined class, GameActivity.java: "public native void nativeVirtualKeyboardChanged(String contents);"
 JNI_METHOD void Java_com_epicgames_ue4_GameActivity_nativeVirtualKeyboardChanged(JNIEnv* jenv, jobject thiz, jstring contents)
 {
-	if (VirtualKeyboardWidget != NULL)
+	if (VirtualKeyboardWidget.IsValid())
 	{
 		const char* javaChars = jenv->GetStringUTFChars(contents, 0);
 
@@ -876,9 +932,10 @@ JNI_METHOD void Java_com_epicgames_ue4_GameActivity_nativeVirtualKeyboardChanged
 		{
 			FGraphEventRef SetWidgetText = FFunctionGraphTask::CreateAndDispatchWhenReady([&]()
 			{
-				if (VirtualKeyboardWidget != NULL)
+				TSharedPtr<IVirtualKeyboardEntry> LockedKeyboardWidget(VirtualKeyboardWidget.Pin());
+				if (LockedKeyboardWidget.IsValid())
 				{
-					VirtualKeyboardWidget->SetTextFromVirtualKeyboard(FText::FromString(FString(UTF8_TO_TCHAR(javaChars))), ETextEntryType::TextEntryUpdated);
+					LockedKeyboardWidget->SetTextFromVirtualKeyboard(FText::FromString(FString(UTF8_TO_TCHAR(javaChars))), ETextEntryType::TextEntryUpdated);
 				}
 			}, TStatId(), NULL, ENamedThreads::GameThread);
 			FTaskGraphInterface::Get().WaitUntilTaskCompletes(SetWidgetText);
@@ -900,15 +957,16 @@ JNI_METHOD void Java_com_epicgames_ue4_GameActivity_nativeVirtualKeyboardSendKey
 JNI_METHOD void Java_com_epicgames_ue4_GameActivity_nativeVirtualKeyboardSendSelection(JNIEnv* jenv, jobject thiz, jint selStart, jint selEnd)
 {
 	// call to set the widget selection on game thread
-	if (VirtualKeyboardWidget != NULL)
+	if (VirtualKeyboardWidget.IsValid())
 	{
 		if (FTaskGraphInterface::IsRunning())
 		{
 			FGraphEventRef SetWidgetSelection = FFunctionGraphTask::CreateAndDispatchWhenReady([&, selStart, selEnd]()
 			{
-				if (VirtualKeyboardWidget != NULL)
+				TSharedPtr<IVirtualKeyboardEntry> LockedKeyboardWidget(VirtualKeyboardWidget.Pin());
+				if (LockedKeyboardWidget.IsValid())
 				{
-					VirtualKeyboardWidget->SetSelectionFromVirtualKeyboard(selStart, selEnd);
+					LockedKeyboardWidget->SetSelectionFromVirtualKeyboard(selStart, selEnd);
 				}
 			}, TStatId(), NULL, ENamedThreads::GameThread);
 		}
@@ -1163,18 +1221,21 @@ bool AndroidThunkCpp_Iap_QueryInAppPurchases(const TArray<FString>& ProductIDs)
 
 		// Populate some java types with the provided product information
 		jobjectArray ProductIDArray = (jobjectArray)Env->NewObjectArray(ProductIDs.Num(), FJavaWrapper::JavaStringClass, NULL);
-		for (uint32 Param = 0; Param < ProductIDs.Num(); Param++)
+		if (ProductIDArray != 0)
 		{
-			jstring StringValue = Env->NewStringUTF(TCHAR_TO_UTF8(*ProductIDs[Param]));
-			Env->SetObjectArrayElement(ProductIDArray, Param, StringValue);
-			Env->DeleteLocalRef(StringValue);
-		}
+			for (uint32 Param = 0; Param < ProductIDs.Num(); Param++)
+			{
+				jstring StringValue = Env->NewStringUTF(TCHAR_TO_UTF8(*ProductIDs[Param]));
+				Env->SetObjectArrayElement(ProductIDArray, Param, StringValue);
+				Env->DeleteLocalRef(StringValue);
+			}
 
-		// Execute the java code for this operation
-		bResult = FJavaWrapper::CallBooleanMethod(Env, FJavaWrapper::GoogleServicesThis, FJavaWrapper::AndroidThunkJava_IapQueryInAppPurchases, ProductIDArray);
-		
-		// clean up references
-		Env->DeleteLocalRef(ProductIDArray);
+			// Execute the java code for this operation
+			bResult = FJavaWrapper::CallBooleanMethod(Env, FJavaWrapper::GoogleServicesThis, FJavaWrapper::AndroidThunkJava_IapQueryInAppPurchases, ProductIDArray);
+
+			// clean up references
+			Env->DeleteLocalRef(ProductIDArray);
+		}
 	}
 
 	return bResult;
