Adding the signatures for FJavaWrapper's CallLongMethod, CallFloatMethod and CallDoubleMethod in the AndroidJNI.h

#4.21
#Android

diff --git a/Engine/Source/Runtime/Launch/Public/Android/AndroidJNI.h b/Engine/Source/Runtime/Launch/Public/Android/AndroidJNI.h
--- a/Engine/Source/Runtime/Launch/Public/Android/AndroidJNI.h
+++ b/Engine/Source/Runtime/Launch/Public/Android/AndroidJNI.h
@@ -133,6 +133,10 @@ public:
 	static void CallVoidMethod(JNIEnv* Env, jobject Object, jmethodID Method, ...);
 	static jobject CallObjectMethod(JNIEnv* Env, jobject Object, jmethodID Method, ...);
 	static int32 CallIntMethod(JNIEnv* Env, jobject Object, jmethodID Method, ...);
+	static int64 CallLongMethod(JNIEnv* Env, jobject Object, jmethodID Method, ...);
+	static float CallFloatMethod(JNIEnv* Env, jobject Object, jmethodID Method, ...);
+	static double CallDoubleMethod(JNIEnv* Env, jobject Object, jmethodID Method, ...);
+	
 	static bool CallBooleanMethod(JNIEnv* Env, jobject Object, jmethodID Method, ...);
 
 	// Delegate that can be registered to that is called when an activity is finished
