#rnx

diff --git a/Engine/Source/Runtime/Launch/Private/HTML5/LaunchHTML5.cpp b/Engine/Source/Runtime/Launch/Private/HTML5/LaunchHTML5.cpp
--- a/Engine/Source/Runtime/Launch/Private/HTML5/LaunchHTML5.cpp
+++ b/Engine/Source/Runtime/Launch/Private/HTML5/LaunchHTML5.cpp
@@ -78,8 +78,9 @@ void HTML5_Init ()
 	emscripten_trace_record_frame_end();
 
 	// main loop
-	emscripten_set_main_loop(HTML5_Tick, 0, false);
-	emscripten_set_main_loop_timing(EM_TIMING_SETIMMEDIATE, 0);
+	emscripten_set_main_loop(HTML5_Tick, 0, true);
+//	emscripten_set_main_loop(HTML5_Tick, 0, false);
+//	emscripten_set_main_loop_timing(EM_TIMING_SETIMMEDIATE, 0);
 	EM_ASM(throw 'SimulateInfiniteLoop');
 }
 
