Fix black screen after splash when not using startup movie or early preload screen
Repplying fix from CL 5257895 in Release-4.22

diff --git a/Engine/Source/Runtime/Launch/Private/LaunchEngineLoop.cpp b/Engine/Source/Runtime/Launch/Private/LaunchEngineLoop.cpp
--- a/Engine/Source/Runtime/Launch/Private/LaunchEngineLoop.cpp
+++ b/Engine/Source/Runtime/Launch/Private/LaunchEngineLoop.cpp
@@ -2291,16 +2291,29 @@ int32 FEngineLoop::PreInit(const TCHAR* CmdLine)
             {
 				SCOPED_BOOT_TIMING("PlayFirstPreLoadScreen");
 
-                // hide splash screen now before playing any movies
-                FPlatformMisc::PlatformHandleSplashScreen(false);
-
                 if (FPreLoadScreenManager::Get())
                 {
 					SCOPED_BOOT_TIMING("PlayFirstPreLoadScreen - FPreLoadScreenManager::Get()->Initialize");
-                    //initialize and play our first Early PreLoad Screen if one is setup
+                    // initialize and play our first Early PreLoad Screen if one is setup
                     FPreLoadScreenManager::Get()->Initialize(SlateRenderer.Get());
-                    FPreLoadScreenManager::Get()->PlayFirstPreLoadScreen(EPreLoadScreenTypes::EarlyStartupScreen);
+
+					if (FPreLoadScreenManager::Get()->HasRegisteredPreLoadScreenType(EPreLoadScreenTypes::EarlyStartupScreen))
+					{
+						// disable the splash before playing the early startup screen
+						FPlatformMisc::PlatformHandleSplashScreen(false);
+	                    FPreLoadScreenManager::Get()->PlayFirstPreLoadScreen(EPreLoadScreenTypes::EarlyStartupScreen);
+	                }
+					else
+					{
+						// no early startup screen, show the splash screen
+						FPlatformMisc::PlatformHandleSplashScreen(true);
+					}
                 }
+				else
+				{
+					// no preload manager, show the splash screen
+					FPlatformMisc::PlatformHandleSplashScreen(true);
+				}
             }
 		}
 		else if ( IsRunningCommandlet() )
