Merge CL 5361873.
Stop using a define to determine whether launch engine loop should mount paks and load shaderlibs.  Instead, see if bundle manager is going to do it.
Fix shader binaries not being fully cached on Android because cache was being opened too early.

diff --git a/Engine/Source/Runtime/Core/Public/GenericPlatform/GenericPlatformInstallBundleManager.h b/Engine/Source/Runtime/Core/Public/GenericPlatform/GenericPlatformInstallBundleManager.h
--- a/Engine/Source/Runtime/Core/Public/GenericPlatform/GenericPlatformInstallBundleManager.h
+++ b/Engine/Source/Runtime/Core/Public/GenericPlatform/GenericPlatformInstallBundleManager.h
@@ -198,6 +198,8 @@ public:
 	virtual void CancelBundle(FName BundleName) = 0;
 
 	virtual TOptional<FInstallBundleStatus> GetBundleProgress(FName BundleName) const = 0;
+
+	virtual bool IsNullInterface() const = 0;
 };
 
 class IPlatformInstallBundleManagerModule : public IModuleInterface
diff --git a/Engine/Source/Runtime/Launch/Private/LaunchEngineLoop.cpp b/Engine/Source/Runtime/Launch/Private/LaunchEngineLoop.cpp
--- a/Engine/Source/Runtime/Launch/Private/LaunchEngineLoop.cpp
+++ b/Engine/Source/Runtime/Launch/Private/LaunchEngineLoop.cpp
@@ -54,6 +54,8 @@
 #include "Misc/NetworkVersion.h"
 #include "Templates/UniquePtr.h"
 
+#include "GenericPlatform/GenericPlatformInstallBundleManager.h"
+
 #if !(IS_PROGRAM || WITH_EDITOR)
 #include "IPlatformFilePak.h"
 #endif
@@ -1595,16 +1597,6 @@ int32 FEngineLoop::PreInit(const TCHAR* CmdLine)
 		}
 	}
 
-	const bool bShouldHandleContentInstalledDuringEarlyLoadScreen = AUTOMATICALLY_HANDLE_INSTALLED_CONTENT_AFTER_EARLY_STARTUP_SCREEN;
-	if (bShouldHandleContentInstalledDuringEarlyLoadScreen)
-	{
-		UE_LOG(LogInit, Verbose, TEXT("Reapplying ini settings after early loading screen."));
-
-		extern CORE_API void RecordApplyCVarSettingsFromIni();
-		SCOPED_BOOT_TIMING("RecordApplyCVarSettingsFromIni");
-		RecordApplyCVarSettingsFromIni();
-	}
-
 #if WITH_ENGINE
 	extern ENGINE_API void InitializeRenderingCVarsCaching();
 	InitializeRenderingCVarsCaching();
@@ -2228,24 +2220,6 @@ int32 FEngineLoop::PreInit(const TCHAR* CmdLine)
 				}
 			}
 
-			//If we are NOT expecting to automatically handle content after our EarlyStartupScreen, we can go ahead and load our game shaders now.
-			//If present they will be loaded, and if not present this will gracefully fail and we can handle loading them in our EarlyStartupScreen.
-			if (!bShouldHandleContentInstalledDuringEarlyLoadScreen)
-			{
-				LLM_SCOPE(ELLMTag::Shaders);
-				SCOPED_BOOT_TIMING("FShaderCodeLibrary::OpenLibrary");
-
-				// Open the game library which contains the material shaders.
-				FShaderCodeLibrary::OpenLibrary(FApp::GetProjectName(), FPaths::ProjectContentDir());
-				for (const FString& RootDir : FPlatformMisc::GetAdditionalRootDirectories())
-				{
-					FShaderCodeLibrary::OpenLibrary(FApp::GetProjectName(), FPaths::Combine(RootDir, FApp::GetProjectName(), TEXT("Content")));
-				}
-
-				// Now our shader code main library is opened, kick off the precompile.
-				FShaderPipelineCache::OpenPipelineFileCache(GMaxRHIShaderPlatform);
-			}
-
 			if (GetMoviePlayer()->HasEarlyStartupMovie())
 			{
 				SCOPED_BOOT_TIMING("EarlyStartupMovie");
@@ -2352,7 +2326,9 @@ int32 FEngineLoop::PreInit(const TCHAR* CmdLine)
 #endif
 
 		//Now that our EarlyStartupScreen is finished, lets take the necessary steps to mount paks, apply .ini cvars, and open the shader libraries if we installed content we expect to handle
-		if (bShouldHandleContentInstalledDuringEarlyLoadScreen)
+		//If using a bundle manager, assume its handling all this stuff and that we don't have to do it.
+		IPlatformInstallBundleManager* BundleManager = FPlatformMisc::GetPlatformInstallBundleManager();
+		if (BundleManager == nullptr || BundleManager->IsNullInterface())
 		{
 			// Mount Paks that were installed during EarlyStartupScreen
 			if (FCoreDelegates::OnMountAllPakFiles.IsBound() )
@@ -2367,17 +2343,6 @@ int32 FEngineLoop::PreInit(const TCHAR* CmdLine)
 				FCoreDelegates::OnMountAllPakFiles.Execute(PakFolders);
 			}
 
-			//Reapply CVars after our EarlyLoadScreen
-			{
-				SCOPED_BOOT_TIMING("ReapplyCVarsFromIniAfterEarlyStartupScreen");
-
-				extern CORE_API void ReapplyRecordedCVarSettingsFromIni();
-				extern CORE_API void DeleteRecordedCVarSettingsFromIni();
-
-				ReapplyRecordedCVarSettingsFromIni();
-				DeleteRecordedCVarSettingsFromIni();
-			}
-
 			//Handle opening shader library after our EarlyLoadScreen
 			{
 				LLM_SCOPE(ELLMTag::Shaders);
diff --git a/Engine/Source/Runtime/NullInstallBundleManager/Source/NullInstallBundleManager.cpp b/Engine/Source/Runtime/NullInstallBundleManager/Source/NullInstallBundleManager.cpp
--- a/Engine/Source/Runtime/NullInstallBundleManager/Source/NullInstallBundleManager.cpp
+++ b/Engine/Source/Runtime/NullInstallBundleManager/Source/NullInstallBundleManager.cpp
@@ -74,6 +74,11 @@ class FNullInstallBundleManager : public IPlatformInstallBundleManager
 		return TOptional<FInstallBundleStatus>();
 	}
 
+	virtual bool IsNullInterface() const override
+	{
+		return true;
+	}
+
 private:
 	
 };
