Cleanup unused define

diff --git a/Engine/Source/Runtime/Launch/Private/LaunchEngineLoop.cpp b/Engine/Source/Runtime/Launch/Private/LaunchEngineLoop.cpp
--- a/Engine/Source/Runtime/Launch/Private/LaunchEngineLoop.cpp
+++ b/Engine/Source/Runtime/Launch/Private/LaunchEngineLoop.cpp
@@ -217,10 +217,6 @@ class FFeedbackContext;
 	#define RHI_COMMAND_LIST_DEBUG_TRACES 0
 #endif
 
-#ifndef AUTOMATICALLY_HANDLE_INSTALLED_CONTENT_AFTER_EARLY_STARTUP_SCREEN
-#define AUTOMATICALLY_HANDLE_INSTALLED_CONTENT_AFTER_EARLY_STARTUP_SCREEN 0
-#endif
-
 #if WITH_ENGINE
 	CSV_DECLARE_CATEGORY_MODULE_EXTERN(CORE_API, Basic);
 #endif
