potential fix for PSO caching crashes while doing the image check (several crashes on IOS)

diff --git a/Engine/Source/Runtime/Launch/Private/LaunchEngineLoop.cpp b/Engine/Source/Runtime/Launch/Private/LaunchEngineLoop.cpp
--- a/Engine/Source/Runtime/Launch/Private/LaunchEngineLoop.cpp
+++ b/Engine/Source/Runtime/Launch/Private/LaunchEngineLoop.cpp
@@ -3403,11 +3403,16 @@ int32 FEngineLoop::Init()
 	// Ready to measure thread heartbeat
 	FThreadHeartBeat::Get().Start();
 
+    {
 #if defined(WITH_CODE_GUARD_HANDLER) && WITH_CODE_GUARD_HANDLER
-    void CheckImageIntegrity();
-    CheckImageIntegrity();
+        FShaderPipelineCache::PauseBatching();
+        void CheckImageIntegrity();
+        CheckImageIntegrity();
+        FShaderPipelineCache::ResumeBatching();
 #endif
-	{
+    }
+    
+    {
 		SCOPED_BOOT_TIMING("FCoreDelegates::OnFEngineLoopInitComplete.Broadcast()");
 		FCoreDelegates::OnFEngineLoopInitComplete.Broadcast();
 	}
