- Fixed a timer logic error for Embedded mode "FirstTicks"

diff --git a/Engine/Source/Runtime/Launch/Private/LaunchEngineLoop.cpp b/Engine/Source/Runtime/Launch/Private/LaunchEngineLoop.cpp
--- a/Engine/Source/Runtime/Launch/Private/LaunchEngineLoop.cpp
+++ b/Engine/Source/Runtime/Launch/Private/LaunchEngineLoop.cpp
@@ -4283,9 +4283,9 @@ void FEngineLoop::Tick()
 #if BUILD_EMBEDDED_APP
 	static double LastSleepTime = FPlatformTime::Seconds();
 	double TimeNow = FPlatformTime::Seconds();
-	if (LastSleepTime - TimeNow >= CVarSecondsBeforeEmbeddedAppSleeps.GetValueOnAnyThread())
+	if (LastSleepTime > 0 && TimeNow - LastSleepTime >= CVarSecondsBeforeEmbeddedAppSleeps.GetValueOnAnyThread())
 	{
-		LastSleepTime = TimeNow;
+		LastSleepTime = 0;
 		FEmbeddedCommunication::AllowSleep(TEXT("FirstTicks"));
 	}
 #endif
