Allow forwarding back button to Java on Android

#ue4
#android

diff --git a/Engine/Build/Android/Java/src/com/epicgames/ue4/NativeCalls.java b/Engine/Build/Android/Java/src/com/epicgames/ue4/NativeCalls.java
--- a/Engine/Build/Android/Java/src/com/epicgames/ue4/NativeCalls.java
+++ b/Engine/Build/Android/Java/src/com/epicgames/ue4/NativeCalls.java
@@ -5,6 +5,9 @@ public class NativeCalls
 	@SuppressWarnings("JniMissingFunction")
 	public static native void HandleCustomTouchEvent(int deviceId, int pointerId, int action, int source, float x, float y);
 
+	@SuppressWarnings("JniMissingFunction")
+	public static native void AllowJavaBackButtonEvent(boolean allow);
+
 	@SuppressWarnings("JniMissingFunction")
 	public static native void CallNativeToEmbedded(String ID, int Priority, String Subsystem, String Command, String[] Params, String RoutingFunction);
 
diff --git a/Engine/Source/Runtime/Launch/Private/Android/LaunchAndroid.cpp b/Engine/Source/Runtime/Launch/Private/Android/LaunchAndroid.cpp
--- a/Engine/Source/Runtime/Launch/Private/Android/LaunchAndroid.cpp
+++ b/Engine/Source/Runtime/Launch/Private/Android/LaunchAndroid.cpp
@@ -173,6 +173,7 @@ static void OnAppCommandCB(struct android_app* app, int32_t cmd); //Lifetime eve
 
 static bool TryIgnoreClick(AInputEvent* event, size_t actionPointer);
 
+bool GAllowJavaBackButtonEvent = false;
 bool GHasInterruptionRequest = false;
 bool GIsInterrupted = false;
 
@@ -1029,6 +1030,12 @@ static int32_t HandleInputCB(struct android_app* app, AInputEvent* event)
 					return 0;
 				}
 			}
+
+			// optionally forward back button
+			if (keyCode == AKEYCODE_BACK && GAllowJavaBackButtonEvent)
+			{
+				return 0;
+			}
 		}
 
 		return 1;
@@ -1399,6 +1406,11 @@ JNI_METHOD void Java_com_epicgames_ue4_NativeCalls_HandleCustomTouchEvent(JNIEnv
 #endif
 }
 
+JNI_METHOD void Java_com_epicgames_ue4_NativeCalls_AllowJavaBackButtonEvent(JNIEnv* jenv, jobject thiz, jboolean allow)
+{
+	GAllowJavaBackButtonEvent = (allow == JNI_TRUE);
+}
+
 bool WaitForAndroidLoseFocusEvent(double TimeoutSeconds)
 {
 	return FAppEventManager::GetInstance()->WaitForEventInQueue(EAppEventState::APP_EVENT_STATE_WINDOW_LOST_FOCUS, TimeoutSeconds);
