[NO NON-PII COMMIT MESSAGE]

diff --git a/Engine/Source/Runtime/Launch/Private/IOS/LaunchIOS.cpp b/Engine/Source/Runtime/Launch/Private/IOS/LaunchIOS.cpp
--- a/Engine/Source/Runtime/Launch/Private/IOS/LaunchIOS.cpp
+++ b/Engine/Source/Runtime/Launch/Private/IOS/LaunchIOS.cpp
@@ -57,7 +57,7 @@ void FAppEntry::Suspend(bool bIsInterrupt)
 			FAudioDevice* AudioDevice = GEngine->GetMainAudioDevice();
 			if (bIsInterrupt && DisableAudioSuspendOnAudioInterruptCvar)
 			{
-				if (FTaskGraphInterface::IsRunning())
+				if (FTaskGraphInterface::IsRunning() && !GIsRequestingExit)
 				{
 					FFunctionGraphTask::CreateAndDispatchWhenReady([AudioDevice]()
 					{
