potential fix for highest rate crashes in 8.30.1 on IOS
they seemed to be caused by shader caching being turned back on, so we are now going to delay the cache start until after those items are complete

diff --git a/Engine/Source/Runtime/Launch/Private/LaunchEngineLoop.cpp b/Engine/Source/Runtime/Launch/Private/LaunchEngineLoop.cpp
--- a/Engine/Source/Runtime/Launch/Private/LaunchEngineLoop.cpp
+++ b/Engine/Source/Runtime/Launch/Private/LaunchEngineLoop.cpp
@@ -3403,12 +3403,11 @@ int32 FEngineLoop::Init()
 	// Ready to measure thread heartbeat
 	FThreadHeartBeat::Get().Start();
 
-    {
+	FShaderPipelineCache::PauseBatching();
+   	{
 #if defined(WITH_CODE_GUARD_HANDLER) && WITH_CODE_GUARD_HANDLER
-        FShaderPipelineCache::PauseBatching();
-        void CheckImageIntegrity();
+         void CheckImageIntegrity();
         CheckImageIntegrity();
-        FShaderPipelineCache::ResumeBatching();
 #endif
     }
     
@@ -3416,6 +3415,7 @@ int32 FEngineLoop::Init()
 		SCOPED_BOOT_TIMING("FCoreDelegates::OnFEngineLoopInitComplete.Broadcast()");
 		FCoreDelegates::OnFEngineLoopInitComplete.Broadcast();
 	}
+	FShaderPipelineCache::ResumeBatching();
 
 #if BUILD_EMBEDDED_APP
 	FEmbeddedCommunication::AllowSleep(TEXT("Startup"));
