Add JNI for routing service intent to native code

#android

diff --git a/Engine/Build/Android/Java/src/com/epicgames/ue4/NativeCalls.java b/Engine/Build/Android/Java/src/com/epicgames/ue4/NativeCalls.java
--- a/Engine/Build/Android/Java/src/com/epicgames/ue4/NativeCalls.java
+++ b/Engine/Build/Android/Java/src/com/epicgames/ue4/NativeCalls.java
@@ -25,4 +25,7 @@ public class NativeCalls
 
 	@SuppressWarnings("JniMissingFunction")
 	public static native void ForwardNotification(String payload);
+
+	@SuppressWarnings("JniMissingFunction")
+	public static native void RouteServiceIntent(String action, String payload);
 }
\ No newline at end of file
diff --git a/Engine/Source/Runtime/Launch/Private/Android/AndroidJNI.cpp b/Engine/Source/Runtime/Launch/Private/Android/AndroidJNI.cpp
--- a/Engine/Source/Runtime/Launch/Private/Android/AndroidJNI.cpp
+++ b/Engine/Source/Runtime/Launch/Private/Android/AndroidJNI.cpp
@@ -46,6 +46,7 @@ extern FString GAPKFilename;
 
 FOnActivityResult FJavaWrapper::OnActivityResultDelegate;
 FOnSafetyNetAttestationResult FJavaWrapper::OnSafetyNetAttestationResultDelegate;
+FOnRouteServiceIntent FJavaWrapper::OnRouteServiceIntentDelegate;
 
 //////////////////////////////////////////////////////////////////////////
 
@@ -1840,6 +1841,26 @@ JNI_METHOD void Java_com_epicgames_ue4_NativeCalls_ForwardNotification(JNIEnv* j
 	//
 }
 
+JNI_METHOD void Java_com_epicgames_ue4_NativeCalls_RouteServiceIntent(JNIEnv* jenv, jobject thiz, jstring InAction, jstring InPayload)
+{
+	// call to OnSafetyNetAttestationResultDelegate on game thread
+	if (FTaskGraphInterface::IsRunning())
+	{
+		const char *nativeAction = jenv->GetStringUTFChars(InAction, 0);
+		FString Action = FString(nativeAction);
+		jenv->ReleaseStringUTFChars(InAction, nativeAction);
+
+		const char *nativePayload = jenv->GetStringUTFChars(InPayload, 0);
+		FString Payload = FString(nativePayload);
+		jenv->ReleaseStringUTFChars(InPayload, nativePayload);
+
+		FGraphEventRef RouteServiceIntentTask = FFunctionGraphTask::CreateAndDispatchWhenReady([Action, Payload]()
+		{
+			FJavaWrapper::OnRouteServiceIntentDelegate.Broadcast(Action, Payload);
+		}, TStatId(), NULL, ENamedThreads::GameThread);
+	}
+}
+
 class FAndroidEmbeddedExec : public FSelfRegisteringExec
 {
 public:
diff --git a/Engine/Source/Runtime/Launch/Public/Android/AndroidJNI.h b/Engine/Source/Runtime/Launch/Public/Android/AndroidJNI.h
--- a/Engine/Source/Runtime/Launch/Public/Android/AndroidJNI.h
+++ b/Engine/Source/Runtime/Launch/Public/Android/AndroidJNI.h
@@ -14,6 +14,7 @@ DECLARE_MULTICAST_DELEGATE_SixParams(FOnActivityResult, JNIEnv *, jobject, jobje
 
 DECLARE_MULTICAST_DELEGATE_ThreeParams(FOnSafetyNetAttestationResult, bool, const FString&, int32);
 
+DECLARE_MULTICAST_DELEGATE_TwoParams(FOnRouteServiceIntent, const FString&, const FString&);
 
 // Define all the Java classes/methods that the game will need to access to
 class FJavaWrapper
@@ -157,6 +158,9 @@ public:
 	// Delegate that can be registered to that is called when an SafetyNet Attestation is finished
 	static FOnSafetyNetAttestationResult OnSafetyNetAttestationResultDelegate;
 
+	// Delegate that can be registered to be called when a service intent is received
+	static FOnRouteServiceIntent OnRouteServiceIntentDelegate;
+
 private:
 
 	/** Find GooglePlay "game services" classes and methods */
