Only write RHI thread time to CSV stats if enabled.

diff --git a/Engine/Source/Runtime/Launch/Private/LaunchEngineLoop.cpp b/Engine/Source/Runtime/Launch/Private/LaunchEngineLoop.cpp
--- a/Engine/Source/Runtime/Launch/Private/LaunchEngineLoop.cpp
+++ b/Engine/Source/Runtime/Launch/Private/LaunchEngineLoop.cpp
@@ -1005,7 +1005,10 @@ static void UpdateCoreCsvStats_EndFrame()
 	CSV_CUSTOM_STAT_GLOBAL(RenderThreadTime, FPlatformTime::ToMilliseconds(GRenderThreadTime), ECsvCustomStatOp::Set);
 	CSV_CUSTOM_STAT_GLOBAL(GameThreadTime, FPlatformTime::ToMilliseconds(GGameThreadTime), ECsvCustomStatOp::Set);
 	CSV_CUSTOM_STAT_GLOBAL(GPUTime, FPlatformTime::ToMilliseconds(GGPUFrameTime), ECsvCustomStatOp::Set);
-	CSV_CUSTOM_STAT_GLOBAL(RHIThreadTime, FPlatformTime::ToMilliseconds(GRHIThreadTime), ECsvCustomStatOp::Set);
+	if (IsRunningRHIInSeparateThread())
+	{
+		CSV_CUSTOM_STAT_GLOBAL(RHIThreadTime, FPlatformTime::ToMilliseconds(GRHIThreadTime), ECsvCustomStatOp::Set);
+	}
 	if (GInputLatencyTime > 0)
 	{
 		CSV_CUSTOM_STAT_GLOBAL(InputLatencyTime, FPlatformTime::ToMilliseconds(GInputLatencyTime), ECsvCustomStatOp::Set);
