Force unhandled exceptions '(ie. exceptions thrown by threads created by third party libraries) to cause the process to exit with exit code 3. This should allow us to account for the number of abnormal exit codes that triggered a crash report.

Tested by spawning a non-UE managed thread (ie. created with the OS CreateThread function rather than FRunnableThread) and immediately triggering an exception from UEngine::Init().

diff --git a/Engine/Source/Runtime/Launch/Private/Windows/LaunchWindows.cpp b/Engine/Source/Runtime/Launch/Private/Windows/LaunchWindows.cpp
--- a/Engine/Source/Runtime/Launch/Private/Windows/LaunchWindows.cpp
+++ b/Engine/Source/Runtime/Launch/Private/Windows/LaunchWindows.cpp
@@ -115,6 +115,8 @@ void SetupWindowsEnvironment( void )
 LONG WINAPI UnhandledException(EXCEPTION_POINTERS *ExceptionInfo)
 {
 	ReportCrash(ExceptionInfo);
+	GIsCriticalError = true;
+	FPlatformMisc::RequestExit(true);
 	return EXCEPTION_CONTINUE_SEARCH;
 }
 
