Bring back config/CVar reloading after patching.
Needs a new build before we can test.

diff --git a/Engine/Source/Runtime/Core/Private/Misc/ConfigCacheIni.cpp b/Engine/Source/Runtime/Core/Private/Misc/ConfigCacheIni.cpp
--- a/Engine/Source/Runtime/Core/Private/Misc/ConfigCacheIni.cpp
+++ b/Engine/Source/Runtime/Core/Private/Misc/ConfigCacheIni.cpp
@@ -4059,8 +4059,8 @@ void ReapplyRecordedCVarSettingsFromIni()
 	for (const FIniToReload& IniToReload : InisToReloadList)
 	{
 		FString TempGlobalName = IniToReload.IniGlobalName;
-		GConfig->LoadGlobalIniFile(TempGlobalName, *IniToReload.IniBaseName, NULL, true);
-		if (TempGlobalName.Compare( IniToReload.IniGlobalName,  ESearchCase::IgnoreCase) != 0 )
+		ensure(GConfig->LoadGlobalIniFile(TempGlobalName, *IniToReload.IniBaseName, nullptr, true));
+		if (TempGlobalName.Compare(IniToReload.IniGlobalName, ESearchCase::IgnoreCase) != 0)
 		{
 			UE_LOG(LogConfig, Warning, TEXT("Tried to reload ini %s final name was %s"), *IniToReload.IniGlobalName, *TempGlobalName );
 		}
diff --git a/Engine/Source/Runtime/Core/Public/GenericPlatform/GenericPlatformInstallBundleManager.h b/Engine/Source/Runtime/Core/Public/GenericPlatform/GenericPlatformInstallBundleManager.h
--- a/Engine/Source/Runtime/Core/Public/GenericPlatform/GenericPlatformInstallBundleManager.h
+++ b/Engine/Source/Runtime/Core/Public/GenericPlatform/GenericPlatformInstallBundleManager.h
@@ -178,6 +178,7 @@ struct FInstallBundleResultInfo
 {
 	FName BundleName;
 	EInstallBundleResult Result = EInstallBundleResult::OK;
+	bool bIsStartup = false;
 
 	// Currently, these just forward BPT Error info
 	FText OptionalErrorText;
diff --git a/Engine/Source/Runtime/Launch/Private/LaunchEngineLoop.cpp b/Engine/Source/Runtime/Launch/Private/LaunchEngineLoop.cpp
--- a/Engine/Source/Runtime/Launch/Private/LaunchEngineLoop.cpp
+++ b/Engine/Source/Runtime/Launch/Private/LaunchEngineLoop.cpp
@@ -54,8 +54,6 @@
 #include "Misc/NetworkVersion.h"
 #include "Templates/UniquePtr.h"
 
-#include "GenericPlatform/GenericPlatformInstallBundleManager.h"
-
 #if !(IS_PROGRAM || WITH_EDITOR)
 #include "IPlatformFilePak.h"
 #endif
@@ -221,6 +219,11 @@ class FFeedbackContext;
 	CSV_DECLARE_CATEGORY_MODULE_EXTERN(CORE_API, Basic);
 #endif
 
+#ifndef WITH_CONFIG_PATCHING
+#define WITH_CONFIG_PATCHING 0
+#endif
+
+
 int32 GUseDisregardForGCOnDedicatedServers = 1;
 static FAutoConsoleVariableRef CVarUseDisregardForGCOnDedicatedServers(
 	TEXT("gc.UseDisregardForGCOnDedicatedServers"),
@@ -1595,6 +1598,14 @@ int32 FEngineLoop::PreInit(const TCHAR* CmdLine)
 		}
 	}
 
+	constexpr bool bWithConfigPatching = (WITH_CONFIG_PATCHING != 0);
+	if (bWithConfigPatching)
+	{
+		UE_LOG(LogInit, Verbose, TEXT("Begin recording CVar changes for config patching."));
+
+		RecordApplyCVarSettingsFromIni();
+	}
+
 #if WITH_ENGINE
 	extern ENGINE_API void InitializeRenderingCVarsCaching();
 	InitializeRenderingCVarsCaching();
@@ -2218,6 +2229,16 @@ int32 FEngineLoop::PreInit(const TCHAR* CmdLine)
 				}
 			}
 
+			if (bWithConfigPatching)
+			{
+				IPlatformInstallBundleManager* BundleManager = FPlatformMisc::GetPlatformInstallBundleManager();
+				if (BundleManager != nullptr && !BundleManager->IsNullInterface())
+				{
+					IPlatformInstallBundleManager::InstallBundleCompleteDelegate.AddRaw(this, &FEngineLoop::OnStartupContentMounted);
+				}
+				// If not using the bundle manager, config will be reloaded after ESP, see below
+			}
+
 			if (GetMoviePlayer()->HasEarlyStartupMovie())
 			{
 				SCOPED_BOOT_TIMING("EarlyStartupMovie");
@@ -2341,6 +2362,15 @@ int32 FEngineLoop::PreInit(const TCHAR* CmdLine)
 				FCoreDelegates::OnMountAllPakFiles.Execute(PakFolders);
 			}
 
+			//Reapply CVars after our EarlyLoadScreen
+			if(bWithConfigPatching)
+			{
+				SCOPED_BOOT_TIMING("ReapplyCVarsFromIniAfterEarlyStartupScreen");
+
+				ReapplyRecordedCVarSettingsFromIni();
+				DeleteRecordedCVarSettingsFromIni();
+			}
+
 			//Handle opening shader library after our EarlyLoadScreen
 			{
 				LLM_SCOPE(ELLMTag::Shaders);
@@ -3435,6 +3465,8 @@ void FEngineLoop::Exit()
 	GIsRunning	= 0;
 	GLogConsole	= nullptr;
 
+	IPlatformInstallBundleManager::InstallBundleCompleteDelegate.RemoveAll(this);
+
 	// shutdown visual logger and flush all data
 #if ENABLE_VISUAL_LOG
 	FVisualLogger::Get().Shutdown();
@@ -3590,6 +3622,16 @@ void FEngineLoop::ProcessLocalPlayerSlateOperations() const
 	}
 }
 
+void FEngineLoop::OnStartupContentMounted(FInstallBundleResultInfo Result)
+{
+	if (Result.bIsStartup && Result.Result == EInstallBundleResult::OK)
+	{
+		ReapplyRecordedCVarSettingsFromIni();
+		DeleteRecordedCVarSettingsFromIni();
+
+		IPlatformInstallBundleManager::InstallBundleCompleteDelegate.RemoveAll(this);
+	}
+}
 
 bool FEngineLoop::ShouldUseIdleMode() const
 {
diff --git a/Engine/Source/Runtime/Launch/Public/LaunchEngineLoop.h b/Engine/Source/Runtime/Launch/Public/LaunchEngineLoop.h
--- a/Engine/Source/Runtime/Launch/Public/LaunchEngineLoop.h
+++ b/Engine/Source/Runtime/Launch/Public/LaunchEngineLoop.h
@@ -8,6 +8,8 @@
 	#include "UnrealEngine.h"
 #endif
 
+#include "GenericPlatform/GenericPlatformInstallBundleManager.h"
+
 class FEngineService;
 class FPendingCleanupObjects;
 class ISessionService;
@@ -116,6 +118,8 @@ private:
 	/** Utility function that processes Slate operations. */
 	void ProcessLocalPlayerSlateOperations() const;
 
+	void OnStartupContentMounted(FInstallBundleResultInfo Result);
+
 protected:
 
 	/** Holds a dynamically expanding array of frame times in milliseconds (if FApp::IsBenchmarking() is set). */
