Add DumpEarlyPakFileReads command to dump all pak reads that occur before it's possible to install patched content to CSV/Log.
This is disabled in shipping.

diff --git a/Engine/Source/Runtime/Core/Private/Misc/ConfigCacheIni.cpp b/Engine/Source/Runtime/Core/Private/Misc/ConfigCacheIni.cpp
--- a/Engine/Source/Runtime/Core/Private/Misc/ConfigCacheIni.cpp
+++ b/Engine/Source/Runtime/Core/Private/Misc/ConfigCacheIni.cpp
@@ -4298,6 +4298,7 @@ FConfigHistoryHelper* ConfigHistoryHelper = nullptr;
 
 void RecordApplyCVarSettingsFromIni()
 {
+	check(IniHistoryHelper == nullptr);
 	IniHistoryHelper = new FCVarIniHistoryHelper();
 }
 
@@ -4341,6 +4342,7 @@ void DeleteRecordedCVarSettingsFromIni()
 void RecordConfigReadsFromIni()
 {
 #if !UE_BUILD_SHIPPING
+	check(ConfigHistoryHelper == nullptr);
 	ConfigHistoryHelper = new FConfigHistoryHelper();
 #endif
 }
@@ -4353,7 +4355,7 @@ void DumpRecordedConfigReadsFromIni()
 #endif
 }
 
-void DeleteRecordConfigReadsFromIni()
+void DeleteRecordedConfigReadsFromIni()
 {
 #if !UE_BUILD_SHIPPING
 	check(ConfigHistoryHelper);
diff --git a/Engine/Source/Runtime/Core/Private/Misc/CoreDelegates.cpp b/Engine/Source/Runtime/Core/Private/Misc/CoreDelegates.cpp
--- a/Engine/Source/Runtime/Core/Private/Misc/CoreDelegates.cpp
+++ b/Engine/Source/Runtime/Core/Private/Misc/CoreDelegates.cpp
@@ -31,6 +31,8 @@ FCoreDelegates::FOnMountPak FCoreDelegates::OnMountPak;
 FCoreDelegates::FOnUnmountPak FCoreDelegates::OnUnmountPak;
 FCoreDelegates::FPakFileMountedDelegate FCoreDelegates::PakFileMountedCallback;
 FCoreDelegates::FNoPakFilesMountedDelegate FCoreDelegates::NoPakFilesMountedDelegate;
+FCoreDelegates::FOnFileOpenedForReadFromPakFile FCoreDelegates::OnFileOpenedForReadFromPakFile;
+
 FCoreDelegates::FOnUserLoginChangedEvent FCoreDelegates::OnUserLoginChangedEvent;
 FCoreDelegates::FOnUserControllerConnectionChange FCoreDelegates::OnControllerConnectionChange;
 FCoreDelegates::FOnUserControllerPairingChange FCoreDelegates::OnControllerPairingChange;
diff --git a/Engine/Source/Runtime/Core/Public/Misc/ConfigCacheIni.h b/Engine/Source/Runtime/Core/Public/Misc/ConfigCacheIni.h
--- a/Engine/Source/Runtime/Core/Public/Misc/ConfigCacheIni.h
+++ b/Engine/Source/Runtime/Core/Public/Misc/ConfigCacheIni.h
@@ -914,4 +914,4 @@ CORE_API void DumpRecordedConfigReadsFromIni();
 /**
  * Helper function to clean up config read history
  */
-CORE_API void DeleteRecordConfigReadsFromIni();
+CORE_API void DeleteRecordedConfigReadsFromIni();
diff --git a/Engine/Source/Runtime/Core/Public/Misc/CoreDelegates.h b/Engine/Source/Runtime/Core/Public/Misc/CoreDelegates.h
--- a/Engine/Source/Runtime/Core/Public/Misc/CoreDelegates.h
+++ b/Engine/Source/Runtime/Core/Public/Misc/CoreDelegates.h
@@ -158,6 +158,10 @@ public:
 	// After an attempt to mount all pak files, but none wre found, this is called
 	static FNoPakFilesMountedDelegate NoPakFilesMountedDelegate;
 
+	// When a file is opened for read from a pak file
+	DECLARE_MULTICAST_DELEGATE_TwoParams(FOnFileOpenedForReadFromPakFile, const TCHAR* /*PakFile*/, const TCHAR* /*FileName*/);
+	static FOnFileOpenedForReadFromPakFile OnFileOpenedForReadFromPakFile;
+
     // Delegate used to register a movie streamer with any movie player modules that bind to this delegate
     // Designed to be called when a platform specific movie streamer plugin starts up so that it doesn't need to implement a register for all movie player plugins
     static FRegisterMovieStreamerDelegate RegisterMovieStreamerDelegate;
diff --git a/Engine/Source/Runtime/Launch/Private/LaunchEngineLoop.cpp b/Engine/Source/Runtime/Launch/Private/LaunchEngineLoop.cpp
--- a/Engine/Source/Runtime/Launch/Private/LaunchEngineLoop.cpp
+++ b/Engine/Source/Runtime/Launch/Private/LaunchEngineLoop.cpp
@@ -884,6 +884,98 @@ void LaunchUpdateMostRecentProjectFile()
 	}
 }
 
+#if !UE_BUILD_SHIPPING
+class FFileInPakFileHistoryHelper
+{
+private:
+	struct FFileInPakFileHistory
+	{
+		FString PakFileName;
+		FString FileName;
+	};
+	friend uint32 GetTypeHash(const FFileInPakFileHistory& H)
+	{
+		uint32 Hash = ::GetTypeHash(H.PakFileName);
+		Hash = HashCombine(Hash, ::GetTypeHash(H.FileName));
+		return Hash;
+	}
+	friend bool operator==(const FFileInPakFileHistory& A, const FFileInPakFileHistory& B)
+	{
+		return A.PakFileName == B.PakFileName && A.FileName == B.FileName;
+	}
+
+	TSet<FFileInPakFileHistory> History;
+
+	void OnFileOpenedForRead(const TCHAR* PakFileName, const TCHAR* FileName)
+	{
+		History.Emplace(FFileInPakFileHistory{ PakFileName, FileName });
+	}
+
+public:
+	FFileInPakFileHistoryHelper()
+	{
+		FCoreDelegates::OnFileOpenedForReadFromPakFile.AddRaw(this, &FFileInPakFileHistoryHelper::OnFileOpenedForRead);
+	}
+
+	~FFileInPakFileHistoryHelper()
+	{
+		FCoreDelegates::OnFileOpenedForReadFromPakFile.RemoveAll(this);
+	}
+
+	void DumpHistory()
+	{
+		const FString SavePath = FPaths::ProjectLogDir() / TEXT("FilesLoadedFromPakFiles.csv");
+
+		FArchive* Writer = IFileManager::Get().CreateFileWriter(*SavePath, FILEWRITE_NoFail);
+
+		auto WriteLine = [Writer](FString&& Line)
+		{
+			UE_LOG(LogInit, Display, TEXT("%s"), *Line);
+			FTCHARToUTF8 UTF8String(*(MoveTemp(Line) + LINE_TERMINATOR));
+			Writer->Serialize((UTF8CHAR*)UTF8String.Get(), UTF8String.Length());
+		};
+
+		UE_LOG(LogInit, Display, TEXT("Dumping History of files read from Paks to %s"), *SavePath);
+		UE_LOG(LogInit, Display, TEXT("Begin History of files read from Paks"));
+		UE_LOG(LogInit, Display, TEXT("------------------------------------------------------"));
+		WriteLine(FString::Printf(TEXT("PakFile, File")));
+		for (const FFileInPakFileHistory& H : History)
+		{
+			WriteLine(FString::Printf(TEXT("%s, %s"), *H.PakFileName, *H.FileName));
+		}
+		UE_LOG(LogInit, Display, TEXT("------------------------------------------------------"));
+		UE_LOG(LogInit, Display, TEXT("End History of files read from Paks"));
+
+		delete Writer;
+		Writer = nullptr;
+	}
+};
+TUniquePtr<FFileInPakFileHistoryHelper> FileInPakFileHistoryHelper;
+#endif // !UE_BUILD_SHIPPING
+
+void RecordFileReadsFromPaks()
+{
+#if !UE_BUILD_SHIPPING
+	FileInPakFileHistoryHelper = MakeUnique<FFileInPakFileHistoryHelper>();
+#endif
+}
+
+void DumpRecordedFileReadsFromPaks()
+{
+#if !UE_BUILD_SHIPPING
+	if (FileInPakFileHistoryHelper)
+	{
+		FileInPakFileHistoryHelper->DumpHistory();
+	}
+#endif
+}
+
+void DeleteRecordedFileReadsFromPaks()
+{
+#if !UE_BUILD_SHIPPING
+	FileInPakFileHistoryHelper = nullptr;
+#endif
+}
 
 /*-----------------------------------------------------------------------------
 	FEngineLoop implementation.
@@ -1599,6 +1691,7 @@ int32 FEngineLoop::PreInit(const TCHAR* CmdLine)
 	}
 
 	const bool bDumpEarlyConfigReads = FParse::Param(FCommandLine::Get(), TEXT("DumpEarlyConfigReads"));
+	const bool bDumpEarlyPakFileReads = FParse::Param(FCommandLine::Get(), TEXT("DumpEarlyPakFileReads"));
 
 	constexpr bool bWithConfigPatching = (WITH_CONFIG_PATCHING != 0);
 	if (bWithConfigPatching)
@@ -1609,6 +1702,10 @@ int32 FEngineLoop::PreInit(const TCHAR* CmdLine)
 		{
 			RecordConfigReadsFromIni();
 		}
+		if (bDumpEarlyPakFileReads)
+		{
+			RecordFileReadsFromPaks();
+		}
 
 		RecordApplyCVarSettingsFromIni();
 	}
@@ -2241,7 +2338,7 @@ int32 FEngineLoop::PreInit(const TCHAR* CmdLine)
 				IPlatformInstallBundleManager* BundleManager = FPlatformMisc::GetPlatformInstallBundleManager();
 				if (BundleManager != nullptr && !BundleManager->IsNullInterface())
 				{
-					IPlatformInstallBundleManager::InstallBundleCompleteDelegate.AddRaw(this, &FEngineLoop::OnStartupContentMounted, bDumpEarlyConfigReads);
+					IPlatformInstallBundleManager::InstallBundleCompleteDelegate.AddRaw(this, &FEngineLoop::OnStartupContentMounted, bDumpEarlyConfigReads, bDumpEarlyPakFileReads);
 				}
 				// If not using the bundle manager, config will be reloaded after ESP, see below
 			}
@@ -2374,14 +2471,7 @@ int32 FEngineLoop::PreInit(const TCHAR* CmdLine)
 			{
 				SCOPED_BOOT_TIMING("ReapplyCVarsFromIniAfterEarlyStartupScreen");
 
-				if (bDumpEarlyConfigReads)
-				{
-					DumpRecordedConfigReadsFromIni();
-					DeleteRecordConfigReadsFromIni();
-				}
-
-				ReapplyRecordedCVarSettingsFromIni();
-				DeleteRecordedCVarSettingsFromIni();
+				HandleConfigReload(bDumpEarlyConfigReads, bDumpEarlyPakFileReads);
 			}
 
 			//Handle opening shader library after our EarlyLoadScreen
@@ -3635,23 +3725,34 @@ void FEngineLoop::ProcessLocalPlayerSlateOperations() const
 	}
 }
 
-void FEngineLoop::OnStartupContentMounted(FInstallBundleResultInfo Result, bool bDumpEarlyConfigReads)
+void FEngineLoop::OnStartupContentMounted(FInstallBundleResultInfo Result, bool bDumpEarlyConfigReads, bool bDumpEarlyPakFileReads)
 {
 	if (Result.bIsStartup && Result.Result == EInstallBundleResult::OK)
 	{
-		if (bDumpEarlyConfigReads)
-		{
-			DumpRecordedConfigReadsFromIni();
-			DeleteRecordConfigReadsFromIni();
-		}
-
-		ReapplyRecordedCVarSettingsFromIni();
-		DeleteRecordedCVarSettingsFromIni();
+		HandleConfigReload(bDumpEarlyConfigReads, bDumpEarlyPakFileReads);
 
 		IPlatformInstallBundleManager::InstallBundleCompleteDelegate.RemoveAll(this);
 	}
 }
 
+void FEngineLoop::HandleConfigReload(bool bDumpEarlyConfigReads, bool bDumpEarlyPakFileReads)
+{
+	if (bDumpEarlyConfigReads)
+	{
+		DumpRecordedConfigReadsFromIni();
+		DeleteRecordedConfigReadsFromIni();
+	}
+
+	if (bDumpEarlyPakFileReads)
+	{
+		DumpRecordedFileReadsFromPaks();
+		DeleteRecordedFileReadsFromPaks();
+	}
+
+	ReapplyRecordedCVarSettingsFromIni();
+	DeleteRecordedCVarSettingsFromIni();
+}
+
 bool FEngineLoop::ShouldUseIdleMode() const
 {
 	static const auto CVarIdleWhenNotForeground = IConsoleManager::Get().FindTConsoleVariableDataInt(TEXT("t.IdleWhenNotForeground"));
diff --git a/Engine/Source/Runtime/Launch/Public/LaunchEngineLoop.h b/Engine/Source/Runtime/Launch/Public/LaunchEngineLoop.h
--- a/Engine/Source/Runtime/Launch/Public/LaunchEngineLoop.h
+++ b/Engine/Source/Runtime/Launch/Public/LaunchEngineLoop.h
@@ -118,7 +118,9 @@ private:
 	/** Utility function that processes Slate operations. */
 	void ProcessLocalPlayerSlateOperations() const;
 
-	void OnStartupContentMounted(FInstallBundleResultInfo Result, bool bDumpEarlyConfigReads);
+	void OnStartupContentMounted(FInstallBundleResultInfo Result, bool bDumpEarlyConfigReads, bool bDumpEarlyPakFileReads);
+
+	void HandleConfigReload(bool bDumpEarlyConfigReads, bool bDumpEarlyPakFileReads);
 
 protected:
 
diff --git a/Engine/Source/Runtime/PakFile/Private/IPlatformFilePak.cpp b/Engine/Source/Runtime/PakFile/Private/IPlatformFilePak.cpp
--- a/Engine/Source/Runtime/PakFile/Private/IPlatformFilePak.cpp
+++ b/Engine/Source/Runtime/PakFile/Private/IPlatformFilePak.cpp
@@ -5930,6 +5930,11 @@ IFileHandle* FPakPlatformFile::OpenRead(const TCHAR* Filename, bool bAllowWrite)
 #endif
 
 		Result = CreatePakFileHandle(Filename, PakFile, &FileEntry);
+
+		if (Result)
+		{
+			FCoreDelegates::OnFileOpenedForReadFromPakFile.Broadcast(*PakFile->GetFilename(), Filename);
+		}
 	}
 	else
 	{
