Fix static analysis warning

diff --git a/Engine/Source/Runtime/Launch/Private/LaunchEngineLoop.cpp b/Engine/Source/Runtime/Launch/Private/LaunchEngineLoop.cpp
--- a/Engine/Source/Runtime/Launch/Private/LaunchEngineLoop.cpp
+++ b/Engine/Source/Runtime/Launch/Private/LaunchEngineLoop.cpp
@@ -1693,7 +1693,13 @@ int32 FEngineLoop::PreInit(const TCHAR* CmdLine)
 	const bool bDumpEarlyConfigReads = FParse::Param(FCommandLine::Get(), TEXT("DumpEarlyConfigReads"));
 	const bool bDumpEarlyPakFileReads = FParse::Param(FCommandLine::Get(), TEXT("DumpEarlyPakFileReads"));
 
-	constexpr bool bWithConfigPatching = (WITH_CONFIG_PATCHING != 0);
+	// Overly verbose to avoid a dumb static analysis warning
+#if WITH_CONFIG_PATCHING
+	constexpr bool bWithConfigPatching = true;
+#else
+	constexpr bool bWithConfigPatching = false;
+#endif
+
 	if (bWithConfigPatching)
 	{
 		UE_LOG(LogInit, Verbose, TEXT("Begin recording CVar changes for config patching."));
