check that we have persistent storage before accessing it for preloaded pak files.

diff --git a/Engine/Source/Runtime/Launch/Private/LaunchEngineLoop.cpp b/Engine/Source/Runtime/Launch/Private/LaunchEngineLoop.cpp
--- a/Engine/Source/Runtime/Launch/Private/LaunchEngineLoop.cpp
+++ b/Engine/Source/Runtime/Launch/Private/LaunchEngineLoop.cpp
@@ -2464,7 +2464,7 @@ int32 FEngineLoop::PreInit(const TCHAR* CmdLine)
 		if (BundleManager == nullptr || BundleManager->IsNullInterface())
 		{
 			// Mount Paks that were installed during EarlyStartupScreen
-			if (FCoreDelegates::OnMountAllPakFiles.IsBound() )
+			if (FCoreDelegates::OnMountAllPakFiles.IsBound() && FPaths::HasProjectPersistentDownloadDir() )
 			{
 				SCOPED_BOOT_TIMING("MountPaksAfterEarlyStartupScreen");
 
