Delay destroy of window to prevent Vulkan hang waiting for new window handle when an enqueued present is flushed on render thread

#android

diff --git a/Engine/Source/Runtime/Launch/Private/Android/AndroidEventManager.cpp b/Engine/Source/Runtime/Launch/Private/Android/AndroidEventManager.cpp
--- a/Engine/Source/Runtime/Launch/Private/Android/AndroidEventManager.cpp
+++ b/Engine/Source/Runtime/Launch/Private/Android/AndroidEventManager.cpp
@@ -129,8 +129,9 @@ void FAppEventManager::Tick()
 					}
 					else
 					{
-						FAndroidAppEntry::DestroyWindow();
-						FAndroidWindow::SetHardwareWindow(NULL);
+						//FAndroidAppEntry::DestroyWindow();
+						//FAndroidWindow::SetHardwareWindow(NULL);
+						bDestroyWindow = true;
 					}
 				}
 			}
