exit app during APP_CMD_DESTROY events.
#rnx

diff --git a/Engine/Source/Runtime/Launch/Private/Android/LaunchAndroid.cpp b/Engine/Source/Runtime/Launch/Private/Android/LaunchAndroid.cpp
--- a/Engine/Source/Runtime/Launch/Private/Android/LaunchAndroid.cpp
+++ b/Engine/Source/Runtime/Launch/Private/Android/LaunchAndroid.cpp
@@ -1451,6 +1451,11 @@ static void OnAppCommandCB(struct android_app* app, int32_t cmd)
 
 
 		FAppEventManager::GetInstance()->EnqueueAppEvent(APP_EVENT_STATE_ON_DESTROY);
+
+		// Exit here, avoids having to unlock the window and letting the RHI's deal with invalid window.
+		extern void AndroidThunkCpp_ForceQuit();
+		AndroidThunkCpp_ForceQuit();
+
 		break;
 	}
 
