[NO NON-PII COMMIT MESSAGE]

diff --git a/Engine/Source/Runtime/Launch/Private/IOS/LaunchIOS.cpp b/Engine/Source/Runtime/Launch/Private/IOS/LaunchIOS.cpp
--- a/Engine/Source/Runtime/Launch/Private/IOS/LaunchIOS.cpp
+++ b/Engine/Source/Runtime/Launch/Private/IOS/LaunchIOS.cpp
@@ -59,7 +59,7 @@ void FAppEntry::Suspend(bool bIsInterrupt)
 	// if background audio is active, then we don't want to do suspend any audio
 	if ([[IOSAppDelegate GetDelegate] IsFeatureActive:EAudioFeature::BackgroundAudio] == false)
 	{
-		if (GEngine && GEngine->GetMainAudioDevice())
+		if (GEngine && GEngine->GetMainAudioDevice() && !GIsRequestingExit)
 		{
 			FAudioDevice* AudioDevice = GEngine->GetMainAudioDevice();
 			if (bIsInterrupt && DisableAudioSuspendOnAudioInterruptCvar)
