UE-64624 App fails to open when packaged for Distribution - Android
#Android
#4.21
As Chris.Babcock mentioned, this is caused by an app pause event triggered when the obb downloader starts and puts the activity in background. Adding a flag to ignore the first app pause in this situation.

diff --git a/Engine/Build/Android/Java/src/com/epicgames/ue4/GameActivity.java.template b/Engine/Build/Android/Java/src/com/epicgames/ue4/GameActivity.java.template
--- a/Engine/Build/Android/Java/src/com/epicgames/ue4/GameActivity.java.template
+++ b/Engine/Build/Android/Java/src/com/epicgames/ue4/GameActivity.java.template
@@ -2790,6 +2790,8 @@ public class GameActivity extends NativeActivity implements SurfaceHolder.Callba
 		}
 		else
 		{
+			nativeOnInitialDownloadStarted();
+
 			// Post the check activity handler here to run after onResume completes
 			Log.debug("==============> Posting request for downloader activity");
 			final Handler downloadHandler = new Handler();
@@ -4014,6 +4016,8 @@ public class GameActivity extends NativeActivity implements SurfaceHolder.Callba
 				errorCode = DOWNLOAD_USER_QUIT;
 			}
 			
+			nativeOnInitialDownloadCompleted();
+
 			HasAllFiles = (errorCode == DOWNLOAD_FILES_PRESENT || errorCode == DOWNLOAD_COMPLETED_OK);
 			
 			if(errorCode == DOWNLOAD_NO_RETURN_CODE 
@@ -5395,6 +5399,9 @@ public class GameActivity extends NativeActivity implements SurfaceHolder.Callba
 
 	public native void nativeOnConfigurationChanged(boolean bPortrait);
 		
+	public native void nativeOnInitialDownloadStarted();
+	public native void nativeOnInitialDownloadCompleted();
+
 	static
 	{
 		System.loadLibrary("gnustl_shared");
diff --git a/Engine/Source/Runtime/Launch/Private/Android/LaunchAndroid.cpp b/Engine/Source/Runtime/Launch/Private/Android/LaunchAndroid.cpp
--- a/Engine/Source/Runtime/Launch/Private/Android/LaunchAndroid.cpp
+++ b/Engine/Source/Runtime/Launch/Private/Android/LaunchAndroid.cpp
@@ -993,6 +993,7 @@ static int32_t HandleInputCB(struct android_app* app, AInputEvent* event)
 }
 
 static bool bShouldRestartFromInterrupt = false;
+static bool bIgnorePauseOnDownloaderStart = false;
 static bool IsStartupMoviePlaying()
 {
 	return GEngine && GEngine->IsInitialized() && GetMoviePlayer() && GetMoviePlayer()->IsStartupMoviePlaying();
@@ -1058,6 +1059,7 @@ static void OnAppCommandCB(struct android_app* app, int32_t cmd)
 		 * input focus.
 		 */
 		// if the app lost focus, avoid unnecessary processing (like monitoring the accelerometer)
+
 		UE_LOG(LogAndroid, Log, TEXT("Case APP_CMD_LOST_FOCUS"));
 		FAppEventManager::GetInstance()->EnqueueAppEvent(APP_EVENT_STATE_WINDOW_LOST_FOCUS, NULL);
 		break;
@@ -1159,7 +1161,7 @@ static void OnAppCommandCB(struct android_app* app, int32_t cmd)
 		FAppEventManager::GetInstance()->EnqueueAppEvent(APP_EVENT_STATE_ON_PAUSE);
 
 		// Restart on resuming if did not complete engine initialization
-		if (!bDidCompleteEngineInit && bDidGainFocus)
+		if (!bDidCompleteEngineInit && bDidGainFocus && !bIgnorePauseOnDownloaderStart)
 		{
 			// only do this if early startup enabled
 			FString *EarlyRestart = FAndroidMisc::GetConfigRulesVariable(TEXT("earlyrestart"));
@@ -1168,6 +1170,7 @@ static void OnAppCommandCB(struct android_app* app, int32_t cmd)
 				bShouldRestartFromInterrupt = true;
 			}
 		}
+		bIgnorePauseOnDownloaderStart = false;
 
 		/*
 		 * On the initial loading the pause method must be called immediately
@@ -1306,6 +1309,18 @@ JNI_METHOD void Java_com_epicgames_ue4_GameActivity_nativeSetAndroidVersionInfor
 	FAndroidMisc::SetVersionInfo( UEAndroidVersion, UEPhoneMake, UEPhoneModel, UEPhoneBuildNumber, UEOSLanguage );
 }
 
+//This function is declared in the Java-defined class, GameActivity.java: "public native void nativeOnInitialDownloadStarted();
+JNI_METHOD void Java_com_epicgames_ue4_GameActivity_nativeOnInitialDownloadStarted(JNIEnv* jenv, jobject thiz)
+{
+	bIgnorePauseOnDownloaderStart = true;
+}
+
+//This function is declared in the Java-defined class, GameActivity.java: "public native void nativeOnInitialDownloadCompleted();
+JNI_METHOD void Java_com_epicgames_ue4_GameActivity_nativeOnInitialDownloadCompleted(JNIEnv* jenv, jobject thiz)
+{
+	bIgnorePauseOnDownloaderStart = false;
+}
+
 bool WaitForAndroidLoseFocusEvent(double TimeoutSeconds)
 {
 	return FAppEventManager::GetInstance()->WaitForEventInQueue(EAppEventState::APP_EVENT_STATE_WINDOW_LOST_FOCUS, TimeoutSeconds);
