Prevent warnings trying to delete log files when two UnrealPak instances are running at once.

diff --git a/Engine/Source/Runtime/Launch/Private/LaunchEngineLoop.cpp b/Engine/Source/Runtime/Launch/Private/LaunchEngineLoop.cpp
--- a/Engine/Source/Runtime/Launch/Private/LaunchEngineLoop.cpp
+++ b/Engine/Source/Runtime/Launch/Private/LaunchEngineLoop.cpp
@@ -4573,9 +4573,12 @@ bool FEngineLoop::AppInit( )
 	// Print all initial startup logging
 	FApp::PrintStartupLogMessages();
 
-	// if a logging build, clear out old log files
+	// if a logging build, clear out old log files. Avoid races when multiple processes are running at once.
 #if !NO_LOGGING
-	FMaintenance::DeleteOldLogs();
+	if (!FParse::Param(FCommandLine::Get(), TEXT("MULTIPROCESS")))
+	{
+		FMaintenance::DeleteOldLogs();
+	}
 #endif
 
 #if !UE_BUILD_SHIPPING
