adding CSV events to track application lifetime events.

diff --git a/Engine/Source/Runtime/Launch/Private/LaunchEngineLoop.cpp b/Engine/Source/Runtime/Launch/Private/LaunchEngineLoop.cpp
--- a/Engine/Source/Runtime/Launch/Private/LaunchEngineLoop.cpp
+++ b/Engine/Source/Runtime/Launch/Private/LaunchEngineLoop.cpp
@@ -984,6 +984,43 @@ static void UpdateCoreCsvStats()
 }
 #endif // WITH_ENGINE && CSV_PROFILER
 
+#if WITH_ENGINE
+namespace AppLifetimeEventCapture
+{
+	static void AppWillDeactivate()
+	{
+		UE_LOG( LogCore, Display, TEXT("AppLifetime: Application will deactivate") );
+		CSV_EVENT_GLOBAL(TEXT("App_WillDeactivate"));
+	}
+
+	static void AppHasReactivated()
+	{
+		UE_LOG( LogCore, Display, TEXT("AppLifetime: Application has reactivated") );
+		CSV_EVENT_GLOBAL(TEXT("App_HasReactivated"));
+	}
+
+	static void AppWillEnterBackground()
+	{
+		UE_LOG( LogCore, Display, TEXT("AppLifetime: Application will enter background") );
+		CSV_EVENT_GLOBAL(TEXT("App_WillEnterBackground"));
+	}
+
+	static void AppHasEnteredForeground()
+	{
+		UE_LOG( LogCore, Display, TEXT("AppLifetime: Application has entered foreground") );
+		CSV_EVENT_GLOBAL(TEXT("App_HasEnteredForeground"));
+	}
+
+	static void Init()
+	{
+		FCoreDelegates::ApplicationWillDeactivateDelegate.AddStatic(AppWillDeactivate);
+		FCoreDelegates::ApplicationHasReactivatedDelegate.AddStatic(AppHasReactivated);
+		FCoreDelegates::ApplicationWillEnterBackgroundDelegate.AddStatic(AppWillEnterBackground);
+		FCoreDelegates::ApplicationHasEnteredForegroundDelegate.AddStatic(AppHasEnteredForeground);
+	}
+}
+#endif //WITH_ENGINE
+
 
 DECLARE_CYCLE_STAT( TEXT( "FEngineLoop::PreInit.AfterStats" ), STAT_FEngineLoop_PreInit_AfterStats, STATGROUP_LoadTime );
 
@@ -1616,6 +1653,10 @@ int32 FEngineLoop::PreInit(const TCHAR* CmdLine)
 	FCsvProfiler::Get()->Init();
 #endif
 
+#if WITH_ENGINE
+	AppLifetimeEventCapture::Init();
+#endif
+
 #if WITH_ENGINE && TRACING_PROFILER
 	FTracingProfiler::Get()->Init();
 #endif
