GEngine NULL check for console commands on iOS

diff --git a/Engine/Source/Runtime/Launch/Private/IOS/IOSAppDelegateConsoleHandling.cpp b/Engine/Source/Runtime/Launch/Private/IOS/IOSAppDelegateConsoleHandling.cpp
--- a/Engine/Source/Runtime/Launch/Private/IOS/IOSAppDelegateConsoleHandling.cpp
+++ b/Engine/Source/Runtime/Launch/Private/IOS/IOSAppDelegateConsoleHandling.cpp
@@ -119,7 +119,7 @@ extern bool GShowSplashScreen;
 {
 	if ([ConsoleCommand length] > 0)
 	{
-		if (self.bEngineInit)
+		if (self.bEngineInit && GEngine != nullptr)
 		{
 			TArray<TCHAR> Ch;
 			Ch.AddZeroed([ConsoleCommand length]);
