Added optimized Windows console window output path to GenericConsoleOutput since this slowed down cooking considerably (2 minutes spent in wprintf alone for one large dataset)

When stdout is attached to a console we use the WriteConsoleW function instead of wprintf since the latter is very slow especially in unbuffered mode which the engine currently configures for stdout (see setvbuf call in LaunchEngineLoop.cpp).

At some point we should reconsider this buffering policy since it's likely to slow down other platforms as well but I wanted to do a safe change for now as I don't yet fully understand why the setvbuf call is there in the first place.

#rn minor

diff --git a/Engine/Source/Runtime/Launch/Private/LaunchEngineLoop.cpp b/Engine/Source/Runtime/Launch/Private/LaunchEngineLoop.cpp
--- a/Engine/Source/Runtime/Launch/Private/LaunchEngineLoop.cpp
+++ b/Engine/Source/Runtime/Launch/Private/LaunchEngineLoop.cpp
@@ -276,10 +276,12 @@ public:
 class FOutputDeviceStdOutput : public FOutputDevice
 {
 public:
-
 	FOutputDeviceStdOutput()
-		: AllowedLogVerbosity(ELogVerbosity::Display)
 	{
+#if PLATFORM_WINDOWS
+		bIsConsoleOutput = IsStdoutAttachedToConsole() && !FParse::Param(FCommandLine::Get(), TEXT("GenericConsoleOutput"));
+#endif
+
 		if (FParse::Param(FCommandLine::Get(), TEXT("AllowStdOutLogVerbosity")))
 		{
 			AllowedLogVerbosity = ELogVerbosity::Log;
@@ -291,10 +293,6 @@ public:
 		}
 	}
 
-	virtual ~FOutputDeviceStdOutput()
-	{
-	}
-
 	virtual bool CanBeUsedOnAnyThread() const override
 	{
 		return true;
@@ -304,18 +302,52 @@ public:
 	{
 		if (Verbosity <= AllowedLogVerbosity)
 		{
+			FString line = FOutputDeviceHelper::FormatLogLine(Verbosity, Category, V, GPrintLogTimes);
+
+#if PLATFORM_WINDOWS
+			if (bIsConsoleOutput)
+			{
+				line.AppendChar('\n');
+
+				WriteConsoleW(GetStdHandle(STD_OUTPUT_HANDLE), *line, line.Len(), NULL, NULL);
+
+				return;
+			}
+
+			// fall through to standard printf path
+#endif
+
 #if PLATFORM_USE_LS_SPEC_FOR_WIDECHAR
 			// printf prints wchar_t strings just fine with %ls, while mixing printf()/wprintf() is not recommended (see https://stackoverflow.com/questions/8681623/printf-and-wprintf-in-single-c-code)
-			printf("%ls\n", *FOutputDeviceHelper::FormatLogLine(Verbosity, Category, V, GPrintLogTimes));
+			printf("%ls\n", *line);
 #else
-			wprintf(TEXT("%s\n"), *FOutputDeviceHelper::FormatLogLine(Verbosity, Category, V, GPrintLogTimes));
+			wprintf(TEXT("%s\n"), *line);
 #endif
+
 			fflush(stdout);
 		}
 	}
 
 private:
-	ELogVerbosity::Type AllowedLogVerbosity;
+	ELogVerbosity::Type AllowedLogVerbosity = ELogVerbosity::Display;
+	bool				bIsConsoleOutput = false;
+
+#if PLATFORM_WINDOWS
+	static bool IsStdoutAttachedToConsole()
+	{
+		HANDLE StdoutHandle = GetStdHandle(STD_OUTPUT_HANDLE);
+		if (StdoutHandle != INVALID_HANDLE_VALUE)
+		{
+			DWORD FileType = GetFileType(StdoutHandle);
+			if (FileType == FILE_TYPE_CHAR)
+			{
+				return true;
+			}
+		}
+
+		return false;
+	}
+#endif
 };
 
 
