PR #4812: Fix inconsistent command-line argument handling under Windows (Contributed by [REDACTED_PERSON])

diff --git a/Engine/Source/Runtime/Launch/Private/Windows/LaunchWindows.cpp b/Engine/Source/Runtime/Launch/Private/Windows/LaunchWindows.cpp
--- a/Engine/Source/Runtime/Launch/Private/Windows/LaunchWindows.cpp
+++ b/Engine/Source/Runtime/Launch/Private/Windows/LaunchWindows.cpp
@@ -7,6 +7,7 @@
 #include "HAL/ExceptionHandling.h"
 #include "HAL/PlatformMallocCrash.h"
 #include "Windows/WindowsHWrapper.h"
+#include <shellapi.h>
 
 #if UE_BUILD_DEBUG
 #include <crtdbg.h>
@@ -158,6 +159,51 @@ LAUNCH_API int32 GuardedMainWrapper( const TCHAR* CmdLine, HINSTANCE hInInstance
 	}
 	return ErrorLevel;
 }
+
+/**
+ * The command-line invocation string, processed using the standard Windows CommandLineToArgvW implementation.
+ * This need to be a static global to avoid object unwinding errors in WinMain when SEH is enabled.
+ */
+static FString GSavedCommandLine;
+
+bool ProcessCommandLine()
+{
+	int argc = 0;
+	LPWSTR* argv = ::CommandLineToArgvW(::GetCommandLineW(), &argc);
+	if (argv != nullptr)
+	{
+		// Reconstruct our command line string in a format suitable for consumption by the FParse class
+		GSavedCommandLine = "";
+		for (int32 Option = 1; Option < argc; Option++)
+		{
+			GSavedCommandLine += TEXT(" ");
+			
+			// Inject quotes in the correct position to preserve arguments containing whitespace
+			FString Temp = argv[Option];
+			if (Temp.Contains(TEXT(" ")))
+			{
+				int32 Quote = 0;
+				if(Temp.StartsWith(TEXT("-")))
+				{
+					int32 Separator;
+					if (Temp.FindChar('=', Separator))
+					{
+						Quote = Separator + 1;
+					}
+				}
+				Temp = Temp.Left(Quote) + TEXT("\"") + Temp.Mid(Quote) + TEXT("\"");
+			}
+			GSavedCommandLine += Temp;
+		}
+		
+		// Free memory allocated for CommandLineToArgvW() arguments
+		::LocalFree(argv);
+		return true;
+	}
+	
+	return false;
+}
+
 int32 WINAPI WinMain( _In_ HINSTANCE hInInstance, _In_opt_ HINSTANCE hPrevInstance, _In_ char*, _In_ int32 nCmdShow )
 {
 	// Setup common Windows settings
@@ -167,6 +213,13 @@ int32 WINAPI WinMain( _In_ HINSTANCE hInInstance, _In_opt_ HINSTANCE hPrevInstan
 	hInstance				= hInInstance;
 	const TCHAR* CmdLine = ::GetCommandLineW();
 
+	// Attempt to process the command-line arguments using the standard Windows implementation
+	// (This ensures behavior parity with other platforms where argc and argv are used.)
+	if ( ProcessCommandLine() )
+	{
+		CmdLine = *GSavedCommandLine;
+	}
+
 #if !(UE_BUILD_SHIPPING && WITH_EDITOR)
 	// Named mutex we use to figure out whether we are the first instance of the game running. This is needed to e.g.
 	// make sure there is no contention when trying to save the shader cache.
