Merging CIS fix for non-development configurations.

#rnx

diff --git a/Engine/Source/Runtime/Launch/Private/LaunchEngineLoop.cpp b/Engine/Source/Runtime/Launch/Private/LaunchEngineLoop.cpp
--- a/Engine/Source/Runtime/Launch/Private/LaunchEngineLoop.cpp
+++ b/Engine/Source/Runtime/Launch/Private/LaunchEngineLoop.cpp
@@ -989,8 +989,10 @@ int32 FEngineLoop::PreInit(const TCHAR* CmdLine)
 	FMemory::SetupTLSCachesOnCurrentThread();
 
 	// Set the flag for whether we've build DebugGame instead of Development. The engine does not know this (whereas the launch module does) because it is always built in development.
+#if UE_BUILD_DEVELOPMENT
 #ifdef UE_BUILD_DEVELOPMENT_WITH_DEBUGGAME
 	FApp::SetDebugGame(UE_BUILD_DEVELOPMENT_WITH_DEBUGGAME != 0);
+#endif
 #endif
 
 	// disable/enable LLM based on commandline
