Add a bCompileAgainstApplicationCore setting to the target rules, which allows compiling out references to the ApplicationCore module (which should only be necessary for applications with a GUI). Removed ApplicationCore from several engine tools and utilities.

diff --git a/Engine/Source/Developer/PakFileUtilities/PakFileUtilities.Build.cs b/Engine/Source/Developer/PakFileUtilities/PakFileUtilities.Build.cs
--- a/Engine/Source/Developer/PakFileUtilities/PakFileUtilities.Build.cs
+++ b/Engine/Source/Developer/PakFileUtilities/PakFileUtilities.Build.cs
@@ -6,7 +6,7 @@ public class PakFileUtilities : ModuleRules
 {
 	public PakFileUtilities(ReadOnlyTargetRules Target) : base(Target)
 	{
-		PrivateDependencyModuleNames.AddRange(new string[] { "Core", "PakFile", "Json", "Projects", "ApplicationCore" });
+		PrivateDependencyModuleNames.AddRange(new string[] { "Core", "PakFile", "Json", "Projects" });
 
         PrivateIncludePathModuleNames.AddRange(
             new string[] {
diff --git a/Engine/Source/Programs/BlankProgram/BlankProgram.Build.cs b/Engine/Source/Programs/BlankProgram/BlankProgram.Build.cs
--- a/Engine/Source/Programs/BlankProgram/BlankProgram.Build.cs
+++ b/Engine/Source/Programs/BlankProgram/BlankProgram.Build.cs
@@ -11,7 +11,6 @@ public class BlankProgram : ModuleRules
 		PrivateIncludePaths.Add("Runtime/Launch/Private");		// For LaunchEngineLoop.cpp include
 
 		PrivateDependencyModuleNames.Add("Core");
-		PrivateDependencyModuleNames.Add("ApplicationCore");
 		PrivateDependencyModuleNames.Add("Projects");
 	}
 }
diff --git a/Engine/Source/Programs/BlankProgram/BlankProgram.Target.cs b/Engine/Source/Programs/BlankProgram/BlankProgram.Target.cs
--- a/Engine/Source/Programs/BlankProgram/BlankProgram.Target.cs
+++ b/Engine/Source/Programs/BlankProgram/BlankProgram.Target.cs
@@ -1,4 +1,4 @@
-﻿// Copyright 1998-2018 Epic Games, Inc. All Rights Reserved.
+// Copyright 1998-2018 Epic Games, Inc. All Rights Reserved.
 
 using UnrealBuildTool;
 using System.Collections.Generic;
@@ -27,6 +27,7 @@ public class BlankProgramTarget : TargetRules
 		// Currently this app is not linking against the engine, so we'll compile out references from Core to the rest of the engine
 		bCompileAgainstEngine = false;
 		bCompileAgainstCoreUObject = false;
+		bCompileAgainstApplicationCore = false;
 
 		// UnrealHeaderTool is a console application, not a Windows app (sets entry point to main(), instead of WinMain())
 		bIsBuildingConsoleApplication = true;
diff --git a/Engine/Source/Programs/BuildPatchTool/BuildPatchTool.Build.cs b/Engine/Source/Programs/BuildPatchTool/BuildPatchTool.Build.cs
--- a/Engine/Source/Programs/BuildPatchTool/BuildPatchTool.Build.cs
+++ b/Engine/Source/Programs/BuildPatchTool/BuildPatchTool.Build.cs
@@ -17,7 +17,6 @@ public class BuildPatchTool : ModuleRules
 			{
 				"Core",
 				"CoreUObject",
-				"ApplicationCore",
 				"BuildPatchServices",
 				"Projects",
 				// The below items are not strictly needed by BPT, but core appears to need them during initialization
diff --git a/Engine/Source/Programs/BuildPatchTool/BuildPatchTool.Target.cs b/Engine/Source/Programs/BuildPatchTool/BuildPatchTool.Target.cs
--- a/Engine/Source/Programs/BuildPatchTool/BuildPatchTool.Target.cs
+++ b/Engine/Source/Programs/BuildPatchTool/BuildPatchTool.Target.cs
@@ -18,6 +18,7 @@ public class BuildPatchToolTarget : TargetRules
 		bBuildEditor = false;
 		bCompileAgainstEngine = false;
 		bCompileAgainstCoreUObject = false;
+		bCompileAgainstApplicationCore = false;
 		bCompileLeanAndMeanUE = true;
 		bUseLoggingInShipping = true;
 		bUseChecksInShipping = true;
diff --git a/Engine/Source/Programs/UnrealBuildTool/Configuration/TargetRules.cs b/Engine/Source/Programs/UnrealBuildTool/Configuration/TargetRules.cs
--- a/Engine/Source/Programs/UnrealBuildTool/Configuration/TargetRules.cs
+++ b/Engine/Source/Programs/UnrealBuildTool/Configuration/TargetRules.cs
@@ -505,6 +505,12 @@ namespace UnrealBuildTool
 		[RequiresUniqueBuildEnvironment]
 		public bool bCompileAgainstCoreUObject = true;
 
+		/// <summary>
+		/// Enabled for builds that need to initialize the ApplicationCore module. Command line utilities do not normally need this.
+		/// </summary>
+		[RequiresUniqueBuildEnvironment]
+		public bool bCompileAgainstApplicationCore = true;
+
 		/// <summary>
 		/// If true, include ADO database support in core.
 		/// </summary>
@@ -1683,6 +1689,11 @@ namespace UnrealBuildTool
 			get { return Inner.bCompileAgainstCoreUObject; }
 		}
 
+		public bool bCompileAgainstApplicationCore
+		{
+			get { return Inner.bCompileAgainstApplicationCore; }
+		}
+
 		public bool bIncludeADO
 		{
 			get { return Inner.bIncludeADO; }
diff --git a/Engine/Source/Programs/UnrealBuildTool/Configuration/UEBuildTarget.cs b/Engine/Source/Programs/UnrealBuildTool/Configuration/UEBuildTarget.cs
--- a/Engine/Source/Programs/UnrealBuildTool/Configuration/UEBuildTarget.cs
+++ b/Engine/Source/Programs/UnrealBuildTool/Configuration/UEBuildTarget.cs
@@ -3733,6 +3733,16 @@ namespace UnrealBuildTool
 				GlobalCompileEnvironment.Definitions.Add("WITH_UNREAL_DEVELOPER_TOOLS=0");
 			}
 
+			// Set a macro to control whether to initialize ApplicationCore. Command line utilities should not generally need this.
+			if (Rules.bCompileAgainstApplicationCore)
+			{
+				GlobalCompileEnvironment.Definitions.Add("WITH_APPLICATION_CORE=1");
+			}
+			else
+			{
+				GlobalCompileEnvironment.Definitions.Add("WITH_APPLICATION_CORE=0");
+			}
+
 			if (Rules.bCompileAgainstCoreUObject)
 			{
 				GlobalCompileEnvironment.Definitions.Add("WITH_COREUOBJECT=1");
diff --git a/Engine/Source/Programs/UnrealHeaderTool/UnrealHeaderTool.Build.cs b/Engine/Source/Programs/UnrealHeaderTool/UnrealHeaderTool.Build.cs
--- a/Engine/Source/Programs/UnrealHeaderTool/UnrealHeaderTool.Build.cs
+++ b/Engine/Source/Programs/UnrealHeaderTool/UnrealHeaderTool.Build.cs
@@ -14,8 +14,7 @@ public class UnrealHeaderTool : ModuleRules
 				"Core",
 				"CoreUObject",
 				"Json",
-				"Projects",
-				"ApplicationCore"
+				"Projects"
 			}
 		);
 	
diff --git a/Engine/Source/Programs/UnrealHeaderTool/UnrealHeaderTool.Target.cs b/Engine/Source/Programs/UnrealHeaderTool/UnrealHeaderTool.Target.cs
--- a/Engine/Source/Programs/UnrealHeaderTool/UnrealHeaderTool.Target.cs
+++ b/Engine/Source/Programs/UnrealHeaderTool/UnrealHeaderTool.Target.cs
@@ -26,6 +26,7 @@ public class UnrealHeaderToolTarget : TargetRules
 
 		// Currently this app is not linking against the engine, so we'll compile out references from Core to the rest of the engine
 		bCompileAgainstEngine = false;
+		bCompileAgainstApplicationCore = false;
 
 		// Force execption handling across all modules.
 		bForceEnableExceptions = true;
diff --git a/Engine/Source/Programs/UnrealPak/UnrealPak.Build.cs b/Engine/Source/Programs/UnrealPak/UnrealPak.Build.cs
--- a/Engine/Source/Programs/UnrealPak/UnrealPak.Build.cs
+++ b/Engine/Source/Programs/UnrealPak/UnrealPak.Build.cs
@@ -8,7 +8,7 @@ public class UnrealPak : ModuleRules
 	{
 		PublicIncludePaths.Add("Runtime/Launch/Public");
 
-		PrivateDependencyModuleNames.AddRange(new string[] { "Core", "PakFile", "Json", "Projects", "ApplicationCore", "PakFileUtilities" });
+		PrivateDependencyModuleNames.AddRange(new string[] { "Core", "PakFile", "Json", "Projects", "PakFileUtilities" });
 
 		PrivateIncludePaths.Add("Runtime/Launch/Private");      // For LaunchEngineLoop.cpp include
 
diff --git a/Engine/Source/Programs/UnrealPak/UnrealPak.Target.cs b/Engine/Source/Programs/UnrealPak/UnrealPak.Target.cs
--- a/Engine/Source/Programs/UnrealPak/UnrealPak.Target.cs
+++ b/Engine/Source/Programs/UnrealPak/UnrealPak.Target.cs
@@ -28,6 +28,7 @@ public class UnrealPakTarget : TargetRules
 		bCompileICU = false;
 
 		// UnrealHeaderTool is a console application, not a Windows app (sets entry point to main(), instead of WinMain())
+		bCompileAgainstApplicationCore = false;
 		bIsBuildingConsoleApplication = true;
 	}
 }
diff --git a/Engine/Source/Programs/UnrealVersionSelector/UnrealVersionSelector.Target.cs b/Engine/Source/Programs/UnrealVersionSelector/UnrealVersionSelector.Target.cs
--- a/Engine/Source/Programs/UnrealVersionSelector/UnrealVersionSelector.Target.cs
+++ b/Engine/Source/Programs/UnrealVersionSelector/UnrealVersionSelector.Target.cs
@@ -30,8 +30,5 @@ public class UnrealVersionSelectorTarget : TargetRules
 		// Currently this app is not linking against the engine, so we'll compile out references from Core to the rest of the engine
 		bCompileAgainstEngine = false;
 		bCompileAgainstCoreUObject = bUsingSlate;
-		
-		// UnrealHeaderTool is a console application, not a Windows app (sets entry point to main(), instead of WinMain())
-		//OutLinkEnvironmentConfiguration.bIsBuildingConsoleApplication = true;
 	}
 }
diff --git a/Engine/Source/Runtime/ApplicationCore/ApplicationCore.Build.cs b/Engine/Source/Runtime/ApplicationCore/ApplicationCore.Build.cs
--- a/Engine/Source/Runtime/ApplicationCore/ApplicationCore.Build.cs
+++ b/Engine/Source/Runtime/ApplicationCore/ApplicationCore.Build.cs
@@ -73,5 +73,9 @@ public class ApplicationCore : ModuleRules
 			PublicIncludePaths.AddRange(new string[] {"Runtime/ApplicationCore/Public/IOS"});
 		}
 
+		if (!Target.bCompileAgainstApplicationCore)
+        {
+			throw new System.Exception("ApplicationCore cannot be used when Target.bCompileAgainstApplicationCore = false.");
+        }
 	}
 }
diff --git a/Engine/Source/Runtime/CoreUObject/CoreUObject.Build.cs b/Engine/Source/Runtime/CoreUObject/CoreUObject.Build.cs
--- a/Engine/Source/Runtime/CoreUObject/CoreUObject.Build.cs
+++ b/Engine/Source/Runtime/CoreUObject/CoreUObject.Build.cs
@@ -22,7 +22,6 @@ public class CoreUObject : ModuleRules
 		PublicDependencyModuleNames.Add("Core");
 
 		PrivateDependencyModuleNames.Add("Projects");
-		PrivateDependencyModuleNames.Add("ApplicationCore");
         PrivateDependencyModuleNames.Add("Json");
 
 	}
diff --git a/Engine/Source/Runtime/CoreUObject/Public/UObject/NoExportTypes.h b/Engine/Source/Runtime/CoreUObject/Public/UObject/NoExportTypes.h
--- a/Engine/Source/Runtime/CoreUObject/Public/UObject/NoExportTypes.h
+++ b/Engine/Source/Runtime/CoreUObject/Public/UObject/NoExportTypes.h
@@ -47,7 +47,7 @@
 
 #include "Internationalization/PolyglotTextData.h"
 
-#include "GenericPlatform/ICursor.h"
+#include "../../../ApplicationCore/Public/GenericPlatform/ICursor.h"
 
 #endif
 
diff --git a/Engine/Source/Runtime/Launch/Private/LaunchEngineLoop.cpp b/Engine/Source/Runtime/Launch/Private/LaunchEngineLoop.cpp
--- a/Engine/Source/Runtime/Launch/Private/LaunchEngineLoop.cpp
+++ b/Engine/Source/Runtime/Launch/Private/LaunchEngineLoop.cpp
@@ -28,8 +28,12 @@
 #include "HAL/FileManagerGeneric.h"
 #include "HAL/ExceptionHandling.h"
 #include "Stats/StatsMallocProfilerProxy.h"
+#if WITH_ENGINE
 #include "HAL/PlatformSplash.h"
+#endif
+#if WITH_APPLICATION_CORE
 #include "HAL/PlatformApplicationMisc.h"
+#endif
 #include "HAL/ThreadManager.h"
 #include "ProfilingDebugging/ExternalProfiler.h"
 #include "Containers/Ticker.h"
@@ -401,7 +405,9 @@ public:
 };
 
 
+#if WITH_APPLICATION_CORE
 static TUniquePtr<FOutputDeviceConsole>	GScopedLogConsole;
+#endif
 static TUniquePtr<FOutputDeviceStdOutput> GScopedStdOut;
 static TUniquePtr<FOutputDeviceTestExit> GScopedTestExit;
 
@@ -1054,8 +1060,10 @@ int32 FEngineLoop::PreInit(const TCHAR* CmdLine)
 		return 1;
 	}
 
+#if WITH_APPLICATION_CORE
 	// Initialize log console here to avoid statics initialization issues when launched from the command line.
 	GScopedLogConsole = TUniquePtr<FOutputDeviceConsole>(FPlatformApplicationMisc::CreateConsoleOutputDevice());
+#endif
 
 	// Always enable the backlog so we get all messages, we will disable and clear it in the game
 	// as soon as we determine whether GIsEditor == false
@@ -1125,8 +1133,13 @@ int32 FEngineLoop::PreInit(const TCHAR* CmdLine)
 	}
 
 	// Output devices.
+#if WITH_APPLICATION_CORE
 	GError = FPlatformApplicationMisc::GetErrorOutputDevice();
 	GWarn = FPlatformApplicationMisc::GetFeedbackContext();
+#else
+	GError = FPlatformOutputDevices::GetError();
+	GWarn = FPlatformOutputDevices::GetFeedbackContext();
+#endif
 
 	// allow the command line to override the platform file singleton
 	bool bFileOverrideFound = false;
@@ -1545,8 +1558,10 @@ int32 FEngineLoop::PreInit(const TCHAR* CmdLine)
 #endif
 	}
 
+#if WITH_APPLICATION_CORE
 	// Get a pointer to the log output device
 	GLogConsole = GScopedLogConsole.Get();
+#endif
 
 	LoadPreInitModules();
 
@@ -1632,7 +1647,9 @@ int32 FEngineLoop::PreInit(const TCHAR* CmdLine)
 
 		// platform specific initialization now that the SystemSettings are loaded
 		FPlatformMisc::PlatformInit();
+#if WITH_APPLICATION_CORE
 		FPlatformApplicationMisc::Init();
+#endif
 		FPlatformMemory::Init();
 	}
 
@@ -2544,7 +2561,9 @@ int32 FEngineLoop::PreInit(const TCHAR* CmdLine)
 #if USE_LOCALIZED_PACKAGE_CACHE
 	FPackageLocalizationManager::Get().InitializeFromDefaultCache();
 #endif	// USE_LOCALIZED_PACKAGE_CACHE
+#if WITH_APPLICATION_CORE
 	FPlatformApplicationMisc::PostInit();
+#endif
 #endif // WITH_ENGINE
 
 	//run automation smoke tests now that everything is setup to run
@@ -3926,7 +3945,9 @@ bool FEngineLoop::AppInit( )
 
 	// Platform specific pre-init.
 	FPlatformMisc::PlatformPreInit();
+#if WITH_APPLICATION_CORE
 	FPlatformApplicationMisc::PreInit();
+#endif
 
 	// Keep track of start time.
 	GSystemStartTime = FDateTime::Now().ToString();
@@ -4308,7 +4329,9 @@ void FEngineLoop::AppExit( )
 
 	UE_LOG(LogExit, Log, TEXT("Exiting."));
 
+#if WITH_APPLICATION_CORE
 	FPlatformApplicationMisc::TearDown();
+#endif
 	FPlatformMisc::PlatformTearDown();
 
 	if (GConfig)
diff --git a/Engine/Source/Runtime/PakFile/PakFile.Build.cs b/Engine/Source/Runtime/PakFile/PakFile.Build.cs
--- a/Engine/Source/Runtime/PakFile/PakFile.Build.cs
+++ b/Engine/Source/Runtime/PakFile/PakFile.Build.cs
@@ -7,6 +7,5 @@ public class PakFile : ModuleRules
 	public PakFile(ReadOnlyTargetRules Target) : base(Target)
 	{
 		PrivateDependencyModuleNames.Add("Core");
-		PrivateDependencyModuleNames.Add("ApplicationCore");
 	}
 }
