- Fixed GIsConsoleExecutable which had stopped working after changes to command line processing in WinMain().

diff --git a/Engine/Source/Runtime/Launch/Private/Launch.cpp b/Engine/Source/Runtime/Launch/Private/Launch.cpp
--- a/Engine/Source/Runtime/Launch/Private/Launch.cpp
+++ b/Engine/Source/Runtime/Launch/Private/Launch.cpp
@@ -123,13 +123,7 @@ int32 GuardedMain( const TCHAR* CmdLine )
 #if PLATFORM_WINDOWS
 	FCString::Strcpy(MiniDumpFilenameW, *FString::Printf(TEXT("unreal-v%i-%s.dmp"), FEngineVersion::Current().GetChangelist(), *FDateTime::Now().ToString()));
 
-	const TCHAR* OrgCmdLine = CmdLine;
-
-	CmdLine = FCommandLine::RemoveExeName(OrgCmdLine);
-
-	const TCHAR* CmdOccurence = FCString::Stristr(OrgCmdLine, TEXT("-cmd"));
-	GIsConsoleExecutable = CmdOccurence != NULL && CmdOccurence < CmdLine;
-	
+	GIsConsoleExecutable = (GetFileType(GetStdHandle(STD_OUTPUT_HANDLE)) == FILE_TYPE_CHAR);
 #endif
 
 	int32 ErrorLevel = EnginePreInit( CmdLine );
