Disable stdout logging if there isn't a valid stdout handle.

diff --git a/Engine/Source/Runtime/Launch/Private/LaunchEngineLoop.cpp b/Engine/Source/Runtime/Launch/Private/LaunchEngineLoop.cpp
--- a/Engine/Source/Runtime/Launch/Private/LaunchEngineLoop.cpp
+++ b/Engine/Source/Runtime/Launch/Private/LaunchEngineLoop.cpp
@@ -300,6 +300,11 @@ public:
 		{
 			AllowedLogVerbosity = ELogVerbosity::All;
 		}
+
+		if (stdout == nullptr)
+		{
+			AllowedLogVerbosity = ELogVerbosity::NoLogging;
+		}
 	}
 
 	virtual bool CanBeUsedOnAnyThread() const override
