Moved the registration of a Ctrl-C handler from WinMain() to EngineLoop::PreInit() so that programs also get the fixed Ctrl-C handler on Windows.

diff --git a/Engine/Source/Runtime/Launch/Private/LaunchEngineLoop.cpp b/Engine/Source/Runtime/Launch/Private/LaunchEngineLoop.cpp
--- a/Engine/Source/Runtime/Launch/Private/LaunchEngineLoop.cpp
+++ b/Engine/Source/Runtime/Launch/Private/LaunchEngineLoop.cpp
@@ -995,6 +995,11 @@ DECLARE_CYCLE_STAT( TEXT( "FEngineLoop::PreInit.AfterStats" ), STAT_FEngineLoop_
 
 int32 FEngineLoop::PreInit(const TCHAR* CmdLine)
 {
+#if PLATFORM_WINDOWS
+	// Register a handler for Ctrl-C so we've effective signal handling from the outset.
+	FWindowsPlatformMisc::SetGracefulTerminationHandler();
+#endif // PLATFORM_WINDOWS
+
 	FMemory::SetupTLSCachesOnCurrentThread();
 
 	// Set the flag for whether we've build DebugGame instead of Development. The engine does not know this (whereas the launch module does) because it is always built in development.
diff --git a/Engine/Source/Runtime/Launch/Private/Windows/LaunchWindows.cpp b/Engine/Source/Runtime/Launch/Private/Windows/LaunchWindows.cpp
--- a/Engine/Source/Runtime/Launch/Private/Windows/LaunchWindows.cpp
+++ b/Engine/Source/Runtime/Launch/Private/Windows/LaunchWindows.cpp
@@ -106,9 +106,6 @@ void SetupWindowsEnvironment( void )
 	// don't fill buffers with 0xfd as we make assumptions for FNames st we only use a fraction of the entire buffer
 	_CrtSetDebugFillThreshold( 0 );
 #endif
-
-	// Register a handler for Ctrl-C so we've effective signal handling from the outset.
-	FWindowsPlatformMisc::SetGracefulTerminationHandler();
 }
 
 /**
