UE-68290 - Dev-Mobile: Compile iOS -- 'TcpConsoleListener.h' file not found

diff --git a/Engine/Source/Runtime/Launch/Private/IOS/TcpConsoleListener.cpp b/Engine/Source/Runtime/Launch/Private/IOS/TcpConsoleListener.cpp
new file mode 100644
--- /dev/null
+++ b/Engine/Source/Runtime/Launch/Private/IOS/TcpConsoleListener.cpp
@@ -0,0 +1,124 @@
+// Copyright 1998-2019 Epic Games, Inc. All Rights Reserved.
+
+#include "TcpConsoleListener.h"
+#include "HAL/RunnableThread.h"
+#include "Common/TcpSocketBuilder.h"
+#include "Common/TcpListener.h"
+
+extern void EnqueueConsoleCommand(uint8 *Command);
+TcpConsoleListener *ConsoleListener = nullptr;
+
+/* TcpConsoleListener structors
+ *****************************************************************************/
+
+TcpConsoleListener::TcpConsoleListener(const FIPv4Endpoint& InListenEndpoint)
+	: ListenEndpoint(InListenEndpoint)
+	, bStopping(false)
+	, Listener(nullptr)
+{
+	UE_LOG(LogTemp, Warning, TEXT("[UE4] Console Listener created!\n"));
+	Thread = FRunnableThread::Create(this, TEXT("TcpConsoleListener"), 128 * 1024, TPri_Normal);
+}
+
+
+TcpConsoleListener::~TcpConsoleListener()
+{
+	if (Listener)
+	{
+		delete Listener;
+		Listener = nullptr;
+	}
+	
+	if (Thread != nullptr)
+	{
+		Thread->Kill(true);
+		delete Thread;
+	}
+}
+
+bool TcpConsoleListener::HandleListenerConnectionAccepted(FSocket* ClientSocket, const FIPv4Endpoint& ClientEndpoint)
+{
+	UE_LOG(LogTemp, Warning, TEXT("[UE4] Got new connection!\n"));
+	Connections.Add(ClientSocket);
+	
+	return true;
+}
+
+/* FRunnable interface
+*****************************************************************************/
+
+void TcpConsoleListener::Exit()
+{
+	// do nothing
+}
+
+
+bool TcpConsoleListener::Init()
+{
+	UE_LOG(LogTemp, Warning, TEXT("[UE4] TCP Listener created!\n"));
+	Listener = new FTcpListener(ListenEndpoint);
+	Listener->OnConnectionAccepted().BindRaw(this, &TcpConsoleListener::HandleListenerConnectionAccepted);
+
+	return true;
+}
+
+
+uint32 TcpConsoleListener::Run()
+{
+	static const SIZE_T CommandSize = 1024;
+	uint8 RecvBuffer[CommandSize];
+	/** Current connections */
+	TArray<FSocket*> ToRemove;
+
+	while (!bStopping)
+	{
+		ToRemove.Empty();
+
+		for (auto Connection : Connections)
+		{
+			uint32 PendingDataSize;
+			if (Connection->HasPendingData(PendingDataSize))
+			{
+				int32 BytesRead = 0;
+				memset(RecvBuffer, 0, CommandSize);
+				if (Connection->Recv(RecvBuffer, CommandSize, BytesRead))
+				{
+					UE_LOG(LogTemp, Warning, TEXT("[UE4] Got TCP console command size %i '%s'"), BytesRead, *FString(UTF8_TO_TCHAR(RecvBuffer)));
+					int i;
+					for (i = 0; i < BytesRead; i++)
+					{
+						UE_LOG(LogTemp, Warning, TEXT("[UE4] Buffer %i"), RecvBuffer[i]);
+					}
+					EnqueueConsoleCommand(RecvBuffer);
+				}
+				else
+				{
+					ToRemove.Add(Connection);
+				}
+			}
+		}
+
+		for (auto& Remove : ToRemove)
+		{
+			Connections.Remove(Remove);
+		}
+		
+		FPlatformProcess::Sleep(0.5f);
+	}
+
+	return 0;
+}
+
+
+void TcpConsoleListener::Stop()
+{
+	bStopping = true;
+
+	if (Listener)
+	{
+		delete Listener;
+		Listener = nullptr;
+	}
+}
+
+
diff --git a/Engine/Source/Runtime/Launch/Private/IOS/TcpConsoleListener.h b/Engine/Source/Runtime/Launch/Private/IOS/TcpConsoleListener.h
new file mode 100644
--- /dev/null
+++ b/Engine/Source/Runtime/Launch/Private/IOS/TcpConsoleListener.h
@@ -0,0 +1,56 @@
+// Copyright 1998-2019 Epic Games, Inc. All Rights Reserved.
+
+#pragma once
+
+#include "CoreMinimal.h"
+#include "Interfaces/IPv4/IPv4Endpoint.h"
+
+class FSocket;
+class FTcpListener;
+
+/**
+ *
+ */
+class TcpConsoleListener
+	: FRunnable
+{
+public:
+
+	/**
+	 * Creates and initializes a new instance.
+	 *
+	 */
+	TcpConsoleListener(const FIPv4Endpoint& InListenEndpoint);
+
+	/** Virtual destructor. */
+	virtual ~TcpConsoleListener();
+
+public:
+
+	//~ FRunnable interface
+
+	virtual void Exit() override;
+	virtual bool Init() override;
+	virtual uint32 Run() override;
+	virtual void Stop() override;
+
+private:
+	
+	/** Callback for accepted connections to the local server. */
+	bool HandleListenerConnectionAccepted(FSocket* ClientSocket, const FIPv4Endpoint& ClientEndpoint);
+
+	/** Current connections */
+	TArray<FSocket*> Connections;
+
+	/** Current settings */
+	FIPv4Endpoint ListenEndpoint;
+
+	/** For the thread */
+	bool bStopping;
+
+	/** Holds the local listener for incoming tunnel connections. */
+	FTcpListener* Listener;
+
+	/** Holds the thread object. */
+	FRunnableThread* Thread;
+};
