BUILD FIX for UE-63767 - iOS Device Output Log does not work

diff --git a/Engine/Source/Runtime/Launch/Private/IOS/IOSAppDelegateConsoleHandling.cpp b/Engine/Source/Runtime/Launch/Private/IOS/IOSAppDelegateConsoleHandling.cpp
--- a/Engine/Source/Runtime/Launch/Private/IOS/IOSAppDelegateConsoleHandling.cpp
+++ b/Engine/Source/Runtime/Launch/Private/IOS/IOSAppDelegateConsoleHandling.cpp
@@ -300,12 +300,14 @@ extern bool GShowSplashScreen;
 
 void EnqueueConsoleCommand(uint8 *Command)
 {
+#if !UE_BUILD_SHIPPING && !PLATFORM_TVOS
 	if (!Command)
 	{
 		return;
 	}
 
 	[[IOSAppDelegate GetDelegate] HandleConsoleCommand:[NSString stringWithUTF8String : (const char *)Command]];
+#endif
 }
 
 @end
