Add some protection to controller polling for Java exceptions

#android

diff --git a/Engine/Build/Android/Java/src/com/epicgames/ue4/GameActivity.java.template b/Engine/Build/Android/Java/src/com/epicgames/ue4/GameActivity.java.template
--- a/Engine/Build/Android/Java/src/com/epicgames/ue4/GameActivity.java.template
+++ b/Engine/Build/Android/Java/src/com/epicgames/ue4/GameActivity.java.template
@@ -4592,7 +4592,7 @@ public class GameActivity extends NativeActivity implements SurfaceHolder.Callba
 		for (int deviceIndex=0; deviceIndex < deviceIds.length; deviceIndex++)
 		{
 			InputDevice inputDevice = InputDevice.getDevice(deviceIds[deviceIndex]);
-			if (inputDevice.getId() == deviceId)
+			if (inputDevice != null && inputDevice.getId() == deviceId)
 			{
 				int vendorId = 0;
 				int productId = 0;
diff --git a/Engine/Source/Runtime/Launch/Private/Android/AndroidJNI.cpp b/Engine/Source/Runtime/Launch/Private/Android/AndroidJNI.cpp
--- a/Engine/Source/Runtime/Launch/Private/Android/AndroidJNI.cpp
+++ b/Engine/Source/Runtime/Launch/Private/Android/AndroidJNI.cpp
@@ -527,7 +527,17 @@ bool AndroidThunkCpp_GetInputDeviceInfo(int32 deviceId, FAndroidInputDeviceInfo
 	if (JNIEnv* Env = FAndroidApplication::GetJavaEnv())
 	{
 		jobject deviceInfo = (jobject)Env->CallObjectMethod(FJavaWrapper::GameActivityThis, FJavaWrapper::AndroidThunkJava_GetInputDeviceInfo, deviceId);
-		if (deviceInfo != 0)
+		if (Env->ExceptionCheck())
+		{
+			Env->ExceptionDescribe();
+			Env->ExceptionClear();
+
+			if (deviceInfo != 0)
+			{
+				Env->DeleteLocalRef(deviceInfo);
+			}
+		}
+		else if (deviceInfo != 0)
 		{
 			bool bIsOptional = false;
 			results.DeviceId = deviceId;
