Disable the Windows application error dialog when running with -unattended. Seeing some build machines running automation tests that aren't cleaning up correctly because there's an open WER prompt that requires a user to click "Close" on.

diff --git a/Engine/Source/Runtime/Launch/Private/Windows/LaunchWindows.cpp b/Engine/Source/Runtime/Launch/Private/Windows/LaunchWindows.cpp
--- a/Engine/Source/Runtime/Launch/Private/Windows/LaunchWindows.cpp
+++ b/Engine/Source/Runtime/Launch/Private/Windows/LaunchWindows.cpp
@@ -220,6 +220,12 @@ int32 WINAPI WinMain( _In_ HINSTANCE hInInstance, _In_opt_ HINSTANCE hPrevInstan
 		CmdLine = *GSavedCommandLine;
 	}
 
+	// If we're running in unattended mode, make sure we never display error dialogs if we crash.
+	if ( FParse::Param( CmdLine, TEXT("unattended") ) )
+	{
+		SetErrorMode(SEM_FAILCRITICALERRORS | SEM_NOGPFAULTERRORBOX | SEM_NOOPENFILEERRORBOX);
+	}
+
 #if !(UE_BUILD_SHIPPING && WITH_EDITOR)
 	// Named mutex we use to figure out whether we are the first instance of the game running. This is needed to e.g.
 	// make sure there is no contention when trying to save the shader cache.
