Merging //UE4/Dev-Main@4732136 to Dev-Rendering (//UE4/Dev-Rendering)
Includes android build fixes.

diff --git a/Engine/Plugins/Developer/Concert/ConcertSync/ConcertSyncClient/Source/ConcertSyncClient/Private/ConcertClientSequencerManager.cpp b/Engine/Plugins/Developer/Concert/ConcertSync/ConcertSyncClient/Source/ConcertSyncClient/Private/ConcertClientSequencerManager.cpp
--- a/Engine/Plugins/Developer/Concert/ConcertSync/ConcertSyncClient/Source/ConcertSyncClient/Private/ConcertClientSequencerManager.cpp
+++ b/Engine/Plugins/Developer/Concert/ConcertSync/ConcertSyncClient/Source/ConcertSyncClient/Private/ConcertClientSequencerManager.cpp
@@ -438,7 +438,7 @@ void FSequencerEventClient::ApplyEventToPlayers(const FConcertSequencerState& Ev
 		if (Sequence && CurrentWorld)
 		{
 			Player = NewObject<ULevelSequencePlayer>((UObject*)GetTransientPackage(), FName("ConcertSequencePlayer"));
-			Player->Initialize(Sequence, CurrentWorld->GetLevel(0), FMovieSceneSequencePlaybackSettings());
+			Player->Initialize(Sequence, CurrentWorld->PersistentLevel, FMovieSceneSequencePlaybackSettings());
 		}
 		SequencePlayers.Add(*EventState.SequenceObjectPath, Player);
 	}
diff --git a/Engine/Plugins/Runtime/AndroidPermission/AndroidPermission.uplugin b/Engine/Plugins/Runtime/AndroidPermission/AndroidPermission.uplugin
--- a/Engine/Plugins/Runtime/AndroidPermission/AndroidPermission.uplugin
+++ b/Engine/Plugins/Runtime/AndroidPermission/AndroidPermission.uplugin
@@ -14,10 +14,6 @@
 	"CanContainContent": false,
 	"IsBetaVersion": true,
 	"Installed": false,
-	"SupportedTargetPlatforms": [
-		"Android",
-		"Lumin"
-	],
 	"Modules": [
 		{
 			"Name": "AndroidPermission",
diff --git a/Engine/Source/Runtime/D3D12RHI/Private/D3D12Allocation.cpp b/Engine/Source/Runtime/D3D12RHI/Private/D3D12Allocation.cpp
--- a/Engine/Source/Runtime/D3D12RHI/Private/D3D12Allocation.cpp
+++ b/Engine/Source/Runtime/D3D12RHI/Private/D3D12Allocation.cpp
@@ -1633,7 +1633,7 @@ FD3D12SegHeap* FD3D12SegList::CreateBackingHeap(
 
 	VERIFYD3D12RESULT(Parent->GetDevice()->CreateHeap(&Desc, IID_PPV_ARGS(&D3DHeap)));
 
-	FD3D12SegHeap* Ret = new FD3D12SegHeap(Parent, VisibleNodeMask, D3DHeap, this, FreeHeaps.Num());
+	FD3D12SegHeap* Ret = new FD3D12SegHeap(Parent, VisibleNodeMask, D3DHeap, HeapSize, this, FreeHeaps.Num());
 	FreeHeaps.Add(Ret);
 	return Ret;
 }
diff --git a/Engine/Source/Runtime/D3D12RHI/Private/D3D12Allocation.h b/Engine/Source/Runtime/D3D12RHI/Private/D3D12Allocation.h
--- a/Engine/Source/Runtime/D3D12RHI/Private/D3D12Allocation.h
+++ b/Engine/Source/Runtime/D3D12RHI/Private/D3D12Allocation.h
@@ -667,6 +667,7 @@ private:
 		FD3D12Device* Parent,
 		FRHIGPUMask VisibileNodeMask,
 		ID3D12Heap* NewHeap,
+		uint64 HeapSize,
 		FD3D12SegList* Owner,
 		uint32 Idx) :
 		FD3D12Heap(Parent, VisibileNodeMask),
@@ -675,6 +676,7 @@ private:
 		FirstFreeOffset(0)
 	{
 		this->SetHeap(NewHeap);
+		BeginTrackingResidency(HeapSize);
 	}
 
 	virtual ~FD3D12SegHeap() = default;
diff --git a/Engine/Source/Runtime/Engine/Private/Animation/AnimSequence.cpp b/Engine/Source/Runtime/Engine/Private/Animation/AnimSequence.cpp
--- a/Engine/Source/Runtime/Engine/Private/Animation/AnimSequence.cpp
+++ b/Engine/Source/Runtime/Engine/Private/Animation/AnimSequence.cpp
@@ -867,7 +867,8 @@ void UAnimSequence::PostEditChangeProperty(FPropertyChangedEvent& PropertyChange
 		PostProcessSequence();
 	}
 
-	if (PropertyChangedEvent.Property != nullptr && PropertyChangedEvent.Property->GetFName() == GET_MEMBER_NAME_CHECKED(UAnimSequence, CurveCompressionSettings))
+	UProperty* Property = PropertyChangedEvent.Property;
+	if (Property != nullptr && Property->GetFName() == GET_MEMBER_NAME_CHECKED(UAnimSequence, CurveCompressionSettings))
 	{
 		RequestSyncAnimRecompression(false);
 	}
diff --git a/Engine/Source/Runtime/Engine/Private/VoiceConfig.cpp b/Engine/Source/Runtime/Engine/Private/VoiceConfig.cpp
--- a/Engine/Source/Runtime/Engine/Private/VoiceConfig.cpp
+++ b/Engine/Source/Runtime/Engine/Private/VoiceConfig.cpp
@@ -35,13 +35,21 @@ int32 UVOIPStatics::GetVoiceSampleRate()
 	}
 
 	static FString DesiredSampleRateStr;
-	bRetrievedSampleRate = true;
-
+	
 	if (GConfig->GetString(TEXT("/Script/Engine.AudioSettings"), TEXT("VoiPSampleRate"), DesiredSampleRateStr, GEngineIni))
 	{
-		SampleRate = (int32)GetEnumValueFromString<EVoiceSampleRate>(TEXT("EVoiceSampleRate"), DesiredSampleRateStr);
+		if (DesiredSampleRateStr.Equals(TEXT("Low16000Hz")))
+		{
+			SampleRate = (int32)EVoiceSampleRate::Low16000Hz;
+		}
+		else if (DesiredSampleRateStr.Equals(TEXT("Normal24000Hz")))
+		{
+			SampleRate = (int32)EVoiceSampleRate::Normal24000Hz;
+		}
+
 		if (SampleRate > 0)
 		{		
+			bRetrievedSampleRate = true;
 			return SampleRate;
 		}
 	}
diff --git a/Engine/Source/Runtime/Engine/Public/Net/VoiceConfig.h b/Engine/Source/Runtime/Engine/Public/Net/VoiceConfig.h
--- a/Engine/Source/Runtime/Engine/Public/Net/VoiceConfig.h
+++ b/Engine/Source/Runtime/Engine/Public/Net/VoiceConfig.h
@@ -42,17 +42,6 @@ namespace MicSilenceDetectionConfig
 
 class USoundAttenuation;
 
-template <typename EnumType>
-static FORCEINLINE EnumType GetEnumValueFromString(const FString& EnumName, const FString& String)
-{
-	UEnum* Enum = FindObject<UEnum>((UObject*) ANY_PACKAGE, *EnumName, true);
-	if (!Enum)
-	{
-		return EnumType(0);
-	}
-	return (EnumType)Enum->GetValueByName(FName(*String));
-}
-
 USTRUCT(BlueprintType)
 struct FVoiceSettings
 {
diff --git a/Engine/Source/Runtime/Engine/Public/Physics/Experimental/PhysInterface_Chaos.h b/Engine/Source/Runtime/Engine/Public/Physics/Experimental/PhysInterface_Chaos.h
--- a/Engine/Source/Runtime/Engine/Public/Physics/Experimental/PhysInterface_Chaos.h
+++ b/Engine/Source/Runtime/Engine/Public/Physics/Experimental/PhysInterface_Chaos.h
@@ -593,6 +593,9 @@ public:
 
     ENGINE_API void DeferredAddCollisionDisableTable(uint32 SkelMeshCompID, TMap<struct FRigidBodyIndexPair, bool> * CollisionDisableTable) {}
     ENGINE_API void DeferredRemoveCollisionDisableTable(uint32 SkelMeshCompID) {}
+
+	void MarkForPreSimKinematicUpdate(USkeletalMeshComponent* InSkelComp, ETeleportType InTeleport, bool bNeedsSkinning) {}
+	void ClearPreSimKinematicUpdate(USkeletalMeshComponent* InSkelComp) {}
     
     void AddPendingOnConstraintBreak(FConstraintInstance* ConstraintInstance, int32 SceneType) {}
     void AddPendingSleepingEvent(FBodyInstance* BI, ESleepEvent SleepEventType, int32 SceneType) {}
diff --git a/Engine/Source/Runtime/Launch/Private/LaunchEngineLoop.cpp b/Engine/Source/Runtime/Launch/Private/LaunchEngineLoop.cpp
--- a/Engine/Source/Runtime/Launch/Private/LaunchEngineLoop.cpp
+++ b/Engine/Source/Runtime/Launch/Private/LaunchEngineLoop.cpp
@@ -984,6 +984,43 @@ static void UpdateCoreCsvStats()
 }
 #endif // WITH_ENGINE && CSV_PROFILER
 
+#if WITH_ENGINE
+namespace AppLifetimeEventCapture
+{
+	static void AppWillDeactivate()
+	{
+		UE_LOG( LogCore, Display, TEXT("AppLifetime: Application will deactivate") );
+		CSV_EVENT_GLOBAL(TEXT("App_WillDeactivate"));
+	}
+
+	static void AppHasReactivated()
+	{
+		UE_LOG( LogCore, Display, TEXT("AppLifetime: Application has reactivated") );
+		CSV_EVENT_GLOBAL(TEXT("App_HasReactivated"));
+	}
+
+	static void AppWillEnterBackground()
+	{
+		UE_LOG( LogCore, Display, TEXT("AppLifetime: Application will enter background") );
+		CSV_EVENT_GLOBAL(TEXT("App_WillEnterBackground"));
+	}
+
+	static void AppHasEnteredForeground()
+	{
+		UE_LOG( LogCore, Display, TEXT("AppLifetime: Application has entered foreground") );
+		CSV_EVENT_GLOBAL(TEXT("App_HasEnteredForeground"));
+	}
+
+	static void Init()
+	{
+		FCoreDelegates::ApplicationWillDeactivateDelegate.AddStatic(AppWillDeactivate);
+		FCoreDelegates::ApplicationHasReactivatedDelegate.AddStatic(AppHasReactivated);
+		FCoreDelegates::ApplicationWillEnterBackgroundDelegate.AddStatic(AppWillEnterBackground);
+		FCoreDelegates::ApplicationHasEnteredForegroundDelegate.AddStatic(AppHasEnteredForeground);
+	}
+}
+#endif //WITH_ENGINE
+
 
 DECLARE_CYCLE_STAT( TEXT( "FEngineLoop::PreInit.AfterStats" ), STAT_FEngineLoop_PreInit_AfterStats, STATGROUP_LoadTime );
 
@@ -1616,6 +1653,10 @@ int32 FEngineLoop::PreInit(const TCHAR* CmdLine)
 	FCsvProfiler::Get()->Init();
 #endif
 
+#if WITH_ENGINE
+	AppLifetimeEventCapture::Init();
+#endif
+
 #if WITH_ENGINE && TRACING_PROFILER
 	FTracingProfiler::Get()->Init();
 #endif
