Don't call LaunchStaticShutdownAfterError() in Mac crash handler to avoid calling malloc/free in PhysX shutdown code

diff --git a/Engine/Source/Runtime/Launch/Private/Mac/LaunchMac.cpp b/Engine/Source/Runtime/Launch/Private/Mac/LaunchMac.cpp
--- a/Engine/Source/Runtime/Launch/Private/Mac/LaunchMac.cpp
+++ b/Engine/Source/Runtime/Launch/Private/Mac/LaunchMac.cpp
@@ -49,7 +49,7 @@ void EngineCrashHandler(const FGenericCrashContext& GenericContext)
 		GError->Flush();
 		GError->HandleError();
 	}
-	LaunchStaticShutdownAfterError();
+
 	return Context.GenerateCrashInfoAndLaunchReporter();
 }
 
