PreLoadScreen changes:
- Move call to SwitchGameWindowToUseGameViewport to happen right after EngineInit instead of after first tick of GameEngine. Fixes cases where we are capturing video and trying to get viewport settings earlier in the process.
- Clean up moved to only happen on the first GameEngine tick isntead of checking for the PreLoadScreenManager every tick
#tests windows editor, windows editor client, ios client

diff --git a/Engine/Source/Runtime/Launch/Private/LaunchEngineLoop.cpp b/Engine/Source/Runtime/Launch/Private/LaunchEngineLoop.cpp
--- a/Engine/Source/Runtime/Launch/Private/LaunchEngineLoop.cpp
+++ b/Engine/Source/Runtime/Launch/Private/LaunchEngineLoop.cpp
@@ -3951,32 +3951,8 @@ void FEngineLoop::Tick()
 		// wait for it to finish before we continue to tick or tick again
 		// We do this right after GEngine->Tick() because that is where user code would initiate a load / movie.
 		{
-            if (FPreLoadScreenManager::Get())
-            {
-                if (FPreLoadScreenManager::Get()->HasRegisteredPreLoadScreenType(EPreLoadScreenTypes::EngineLoadingScreen))
-                {
-                    //Wait for any Engine Loading Screen to stop
-                    if (FPreLoadScreenManager::Get()->HasActivePreLoadScreenType(EPreLoadScreenTypes::EngineLoadingScreen))
-                    {
-                        FPreLoadScreenManager::Get()->WaitForEngineLoadingScreenToFinish();
-                    }
-
-                    //Switch Game Window Back
-                    UGameEngine* GameEngine = Cast<UGameEngine>(GEngine);
-                    if (GameEngine)
-                    {
-                        GameEngine->SwitchGameWindowToUseGameViewport();
-                    }
-                }
-                
-                //Destroy / Clean Up PreLoadScreenManager as we are now done
-                FPreLoadScreenManager::Destroy();
-            }
-			else
-			{
-				QUICK_SCOPE_CYCLE_COUNTER(STAT_FEngineLoop_WaitForMovieToFinish);
-				GetMoviePlayer()->WaitForMovieToFinish(true);
-			}
+            QUICK_SCOPE_CYCLE_COUNTER(STAT_FEngineLoop_WaitForMovieToFinish);
+			GetMoviePlayer()->WaitForMovieToFinish(true);
 		}
 
 		if (GShaderCompilingManager)
diff --git a/Engine/Source/Runtime/PreLoadScreen/Private/PreLoadScreenManager.cpp b/Engine/Source/Runtime/PreLoadScreen/Private/PreLoadScreenManager.cpp
--- a/Engine/Source/Runtime/PreLoadScreen/Private/PreLoadScreenManager.cpp
+++ b/Engine/Source/Runtime/PreLoadScreen/Private/PreLoadScreenManager.cpp
@@ -58,8 +58,12 @@ void FPreLoadScreenManager::Initialize(FSlateRenderer& InSlateRenderer)
         MainWindow = GameWindow;
 
         WidgetRenderer = MakeShared<FPreLoadSlateWidgetRenderer, ESPMode::ThreadSafe>(GameWindow, VirtualRenderWindow, &InSlateRenderer);
+
+		FCoreDelegates::OnPostEngineInit.AddRaw(this, &FPreLoadScreenManager::OnPostEngineInit);
     }
 
+	FCoreDelegates::OnBeginFrame.AddRaw(this, &FPreLoadScreenManager::OnFirstEngineFrame);
+
     LastRenderTickTime = FPlatformTime::Seconds();
     LastTickTime = FPlatformTime::Seconds();
 }
@@ -78,6 +82,30 @@ void FPreLoadScreenManager::UnRegisterPreLoadScreen(TSharedPtr<IPreLoadScreen> P
     }
 }
 
+void FPreLoadScreenManager::OnPostEngineInit()
+{
+	FCoreDelegates::OnPostEngineInit.RemoveAll(this);
+	SetupSceneViewport();
+}
+
+void FPreLoadScreenManager::SetupSceneViewport()
+{
+	UGameEngine* GameEngine = Cast<UGameEngine>(GEngine);
+	if (GameEngine)
+	{
+		GameEngine->SwitchGameWindowToUseGameViewport();
+	}
+}
+
+void FPreLoadScreenManager::OnFirstEngineFrame()
+{
+	//Only want to run this on the first Engine Frame
+	FCoreDelegates::OnBeginFrame.RemoveAll(this);
+
+	//Destroy / Clean Up as we are now done
+	Destroy();
+}
+
 void FPreLoadScreenManager::PlayFirstPreLoadScreen(EPreLoadScreenTypes PreLoadScreenTypeToPlay)
 {
     for (int PreLoadScreenIndex = 0; PreLoadScreenIndex < PreLoadScreens.Num(); ++PreLoadScreenIndex)
diff --git a/Engine/Source/Runtime/PreLoadScreen/Public/PreLoadScreenManager.h b/Engine/Source/Runtime/PreLoadScreen/Public/PreLoadScreenManager.h
--- a/Engine/Source/Runtime/PreLoadScreen/Public/PreLoadScreenManager.h
+++ b/Engine/Source/Runtime/PreLoadScreen/Public/PreLoadScreenManager.h
@@ -75,6 +75,9 @@ public:
     void SetEngineLoadingComplete(bool IsEngineLoadingFinished = true);
     bool IsEngineLoadingComplete() const { return bIsEngineLoadingComplete; }
 
+	void OnPostEngineInit();
+	void SetupSceneViewport();
+
     void CleanUpResources();
 
     static void EnableRendering(bool bEnabled);
@@ -103,6 +106,8 @@ protected:
     //Singleton Instance
     static TSharedPtr<FPreLoadScreenManager> Instance;
 
+	void OnFirstEngineFrame();
+
     void BeginPlay();
 
     /*** These functions describe the flow for an EarlyPreLoadScreen where everything is blocking waiting on a call to StopPreLoadScreen ***/
