Disable custom touch event injection

#android

diff --git a/Engine/Source/Runtime/Launch/Private/Android/LaunchAndroid.cpp b/Engine/Source/Runtime/Launch/Private/Android/LaunchAndroid.cpp
--- a/Engine/Source/Runtime/Launch/Private/Android/LaunchAndroid.cpp
+++ b/Engine/Source/Runtime/Launch/Private/Android/LaunchAndroid.cpp
@@ -50,6 +50,10 @@ static GetAxesType GetAxes = NULL;
 #define AMOTION_EVENT_AXIS_RELATIVE_Y 28
 #endif
 
+#ifndef ANDROID_ALLOWCUSTOMTOUCHEVENT
+#define ANDROID_ALLOWCUSTOMTOUCHEVENT	0
+#endif
+
 // List of default axes to query for each controller
 // Ideal solution is to call out to Java and enumerate the list of axes.
 static const int32_t AxisList[] =
@@ -1358,6 +1362,7 @@ JNI_METHOD void Java_com_epicgames_ue4_GameActivity_nativeOnInitialDownloadCompl
 bool GAllowCustomInput = true;
 JNI_METHOD void Java_com_epicgames_ue4_NativeCalls_HandleCustomTouchEvent(JNIEnv* jenv, jobject thiz, jint deviceId, jint pointerId, jint action, jint soucre, jfloat x, jfloat y)
 {
+#if ANDROID_ALLOWCUSTOMTOUCHEVENT
 	// make sure fake input is allowed, so hacky Java can't run bots
 	if (!GAllowCustomInput)
 	{
@@ -1387,6 +1392,7 @@ JNI_METHOD void Java_com_epicgames_ue4_NativeCalls_HandleCustomTouchEvent(JNIEnv
 
 	UE_LOG(LogAndroid, Verbose, TEXT("Handle custom touch event %d (%d) x=%f y=%f"), TouchMessage.Type, action, x, y);
 	FAndroidInputInterface::QueueTouchInput(TouchesArray);
+#endif
 }
 
 bool WaitForAndroidLoseFocusEvent(double TimeoutSeconds)
