Fix console command race condition

diff --git a/Engine/Source/Runtime/Launch/Private/Android/LaunchAndroid.cpp b/Engine/Source/Runtime/Launch/Private/Android/LaunchAndroid.cpp
--- a/Engine/Source/Runtime/Launch/Private/Android/LaunchAndroid.cpp
+++ b/Engine/Source/Runtime/Launch/Private/Android/LaunchAndroid.cpp
@@ -34,6 +34,7 @@
 #include "MoviePlayer.h"
 #include "PreLoadScreenManager.h"
 #include "Misc/EmbeddedCommunication.h"
+#include "Async/Async.h"
 #include <jni.h>
 #include <android/sensor.h>
 
@@ -1312,11 +1313,14 @@ JNI_METHOD void Java_com_epicgames_ue4_GameActivity_nativeOnConfigurationChanged
 //This function is declared in the Java-defined class, GameActivity.java: "public native void nativeConsoleCommand(String commandString);"
 JNI_METHOD void Java_com_epicgames_ue4_GameActivity_nativeConsoleCommand(JNIEnv* jenv, jobject thiz, jstring commandString)
 {
-	auto Command = FJavaHelper::FStringFromParam(jenv, commandString);
-	
+	FString Command = FJavaHelper::FStringFromParam(jenv, commandString);
 	if (GEngine != NULL)
 	{
-		GEngine->DeferredCommands.Add(Command);
+		// Run on game thread to avoid race condition with DeferredCommands
+		AsyncTask(ENamedThreads::GameThread, [Command]()
+		{
+			GEngine->DeferredCommands.Add(Command);
+		});
 	}
 	else
 	{
