UE-65414 ARCore - Rotating app during bootup cause scaling issues on startup
-If we skip a window create because a window destroy event came in before we could process it we now destroy the window and pause rendering so that when a window is created it builds up all state fresh.  Previously we sometimes ended up with the frame buffer oriented incorrectly or some undiagnosed problem which resulted in an all black screen.

diff --git a/Engine/Source/Runtime/Launch/Private/Android/AndroidEventManager.cpp b/Engine/Source/Runtime/Launch/Private/Android/AndroidEventManager.cpp
--- a/Engine/Source/Runtime/Launch/Private/Android/AndroidEventManager.cpp
+++ b/Engine/Source/Runtime/Launch/Private/Android/AndroidEventManager.cpp
@@ -28,6 +28,39 @@ FAppEventManager* FAppEventManager::GetInstance()
 	return sInstance;
 }
 
+static const TCHAR* GetAppEventName(EAppEventState State)
+{
+	const TCHAR* Names[] = {
+		TEXT("APP_EVENT_STATE_WINDOW_CREATED"),
+		TEXT("APP_EVENT_STATE_WINDOW_RESIZED"),
+		TEXT("APP_EVENT_STATE_WINDOW_CHANGED"),
+		TEXT("APP_EVENT_STATE_WINDOW_DESTROYED"),
+		TEXT("APP_EVENT_STATE_WINDOW_REDRAW_NEEDED"),
+		TEXT("APP_EVENT_STATE_ON_DESTROY"),
+		TEXT("APP_EVENT_STATE_ON_PAUSE"),
+		TEXT("APP_EVENT_STATE_ON_RESUME"),
+		TEXT("APP_EVENT_STATE_ON_STOP"),
+		TEXT("APP_EVENT_STATE_ON_START"),
+		TEXT("APP_EVENT_STATE_WINDOW_LOST_FOCUS"),
+		TEXT("APP_EVENT_STATE_WINDOW_GAINED_FOCUS"),
+		TEXT("APP_EVENT_STATE_SAVE_STATE")
+		};
+
+
+	if (State == APP_EVENT_STATE_INVALID)
+	{
+		return TEXT("APP_EVENT_STATE_INVALID");
+	}
+	else if (State > APP_EVENT_STATE_SAVE_STATE || State < 0)
+	{
+		return TEXT("UnknownEAppEventStateValue");
+	}
+	else
+	{
+		return Names[State];
+	}
+}
+
 
 void FAppEventManager::Tick()
 {
@@ -48,8 +81,22 @@ void FAppEventManager::Tick()
 			{
 				bCreateWindow = true;
 				PendingWindow = (ANativeWindow*)Event.Data;
+				FPlatformMisc::LowLevelOutputDebugStringf(TEXT("APP_EVENT_STATE_WINDOW_CREATED %d, %d, %d, %d"), int(bDestroyWindowPending), int(bRunning), int(bHaveWindow), int(bHaveGame));
+
+			}
+			else
+			{
+				// If we skip a create window because it is destroyed before we created the window *something* gets out of sync
+				// resulting in either one buffer still in the wrong orienation or just a black screen.
+				// So reset everything here. When the new window is created we will recover successfully.
+
+				FAndroidAppEntry::DestroyWindow();
+				FAndroidWindow::SetHardwareWindow(NULL);
+
+				PauseRendering();
+				PauseAudio();
+				FPlatformMisc::LowLevelOutputDebugStringf(TEXT("APP_EVENT_STATE_WINDOW_CREATED window creation skipped because a destroy is pending %d, %d, %d, %d"), int(bDestroyWindowPending), int(bRunning), int(bHaveWindow), int(bHaveGame));
 
-				FPlatformMisc::LowLevelOutputDebugStringf(TEXT("APP_EVENT_STATE_WINDOW_CREATED, %d, %d, %d"), int(bRunning), int(bHaveWindow), int(bHaveGame));
 			}
 			break;
 		
@@ -529,7 +576,7 @@ void FAppEventManager::EnqueueAppEvent(EAppEventState InState, void* InData)
 	rc = pthread_mutex_unlock(&QueueMutex);
 	check(rc == 0);
 
-	FPlatformMisc::LowLevelOutputDebugStringf(TEXT("LogAndroidEvents: EnqueueAppEvent : %u, %u, tid = %d"), InState, (uintptr_t)InData, gettid());
+	FPlatformMisc::LowLevelOutputDebugStringf(TEXT("LogAndroidEvents::EnqueueAppEvent : %u, %u, tid = %d, %s"), InState, (uintptr_t)InData, gettid(), GetAppEventName(InState));
 }
 
 
@@ -544,7 +591,7 @@ FAppEventData FAppEventManager::DequeueAppEvent()
 	rc = pthread_mutex_unlock(&QueueMutex);
 	check(rc == 0);
 
-	UE_LOG(LogAndroidEvents, Display, TEXT("DequeueAppEvent : %u, %u"), OutData.State, (uintptr_t)OutData.Data)
+	UE_LOG(LogAndroidEvents, Display, TEXT("LogAndroidEvents::DequeueAppEvent : %u, %u, %s"), OutData.State, (uintptr_t)OutData.Data, GetAppEventName(OutData.State))
 
 	return OutData;
 }
