#UE4 Fix a crash related to StartFrame. WorldContext worlds may be null during tick when starting PIE.

diff --git a/Engine/Source/Runtime/Launch/Private/LaunchEngineLoop.cpp b/Engine/Source/Runtime/Launch/Private/LaunchEngineLoop.cpp
--- a/Engine/Source/Runtime/Launch/Private/LaunchEngineLoop.cpp
+++ b/Engine/Source/Runtime/Launch/Private/LaunchEngineLoop.cpp
@@ -3930,12 +3930,15 @@ void FEngineLoop::Tick()
 		for (const FWorldContext& Context : GEngine->GetWorldContexts())
 		{
 			UWorld* CurrentWorld = Context.World();
-			FSceneInterface* Scene = CurrentWorld->Scene;
-
-			ENQUEUE_RENDER_COMMAND(SceneStartFrame)([Scene](FRHICommandListImmediate& RHICmdList)
+			if (CurrentWorld)
 			{
-				Scene->StartFrame();
-			});
+				FSceneInterface* Scene = CurrentWorld->Scene;
+
+				ENQUEUE_RENDER_COMMAND(SceneStartFrame)([Scene](FRHICommandListImmediate& RHICmdList)
+				{
+					Scene->StartFrame();
+				});
+			}
 		}
 
 #if !UE_SERVER && WITH_ENGINE
