Disabling CPU usage CSV metric in shipping builds as it's still too unreliable at this point but wanted for internal testing.

diff --git a/Engine/Source/Runtime/Launch/Private/LaunchEngineLoop.cpp b/Engine/Source/Runtime/Launch/Private/LaunchEngineLoop.cpp
--- a/Engine/Source/Runtime/Launch/Private/LaunchEngineLoop.cpp
+++ b/Engine/Source/Runtime/Launch/Private/LaunchEngineLoop.cpp
@@ -988,7 +988,7 @@ bool IsServerDelegateForOSS(FName WorldContextHandle)
 #if WITH_ENGINE && CSV_PROFILER
 static void UpdateCoreCsvStats_BeginFrame()
 {
-#if PLATFORM_WINDOWS
+#if PLATFORM_WINDOWS && !UE_BUILD_SHIPPING
 	if (FCsvProfiler::Get()->IsCapturing())
 	{
 		const uint32 ProcessId = (uint32)GetCurrentProcessId();
