Fixing DCC build failure on shared pointer in PreInit

diff --git a/Engine/Source/Runtime/Launch/Private/LaunchEngineLoop.cpp b/Engine/Source/Runtime/Launch/Private/LaunchEngineLoop.cpp
--- a/Engine/Source/Runtime/Launch/Private/LaunchEngineLoop.cpp
+++ b/Engine/Source/Runtime/Launch/Private/LaunchEngineLoop.cpp
@@ -1291,8 +1291,14 @@ int32 FEngineLoop::PreInit(const TCHAR* CmdLine)
 // This shader file is in the plugin, but it is compiled even if the plugin is disabled so we can't have the plugin add this mapping.
 // This is a temporary workaround.
 #if defined(MORPHEUS_ENGINE_DISTORTION) && MORPHEUS_ENGINE_DISTORTION 
-	FString PluginShaderDir = FPaths::Combine(IPluginManager::Get().FindPlugin(TEXT("Morpheus"))->GetBaseDir(), TEXT("Shaders"));
-	AddShaderSourceDirectoryMapping(TEXT("/Plugin/Morpheus"), PluginShaderDir);
+	{
+		TSharedPtr<IPlugin> Plugin = IPluginManager::Get().FindPlugin(TEXT("Morpheus"));
+		if (Plugin.IsValid())
+		{
+			FString PluginShaderDir = FPaths::Combine(Plugin->GetBaseDir(), TEXT("Shaders"));
+			AddShaderSourceDirectoryMapping(TEXT("/Plugin/Morpheus"), PluginShaderDir);
+		}
+	}
 #endif
 
 	TArray<FString> Tokens;
