Add functions for native code to log using embedded UE4

diff --git a/Engine/Build/Android/Java/src/com/epicgames/ue4/NativeCalls.java b/Engine/Build/Android/Java/src/com/epicgames/ue4/NativeCalls.java
--- a/Engine/Build/Android/Java/src/com/epicgames/ue4/NativeCalls.java
+++ b/Engine/Build/Android/Java/src/com/epicgames/ue4/NativeCalls.java
@@ -28,4 +28,16 @@ public class NativeCalls
 
 	@SuppressWarnings("JniMissingFunction")
 	public static native void RouteServiceIntent(String action, String payload);
+
+	@SuppressWarnings("JniMissingFunction")
+	public static native void UELogError(String msg);
+
+	@SuppressWarnings("JniMissingFunction")
+	public static native void UELogWarning(String msg);
+
+	@SuppressWarnings("JniMissingFunction")
+	public static native void UELogLog(String msg);
+
+	@SuppressWarnings("JniMissingFunction")
+	public static native void UELogVerbose(String msg);
 }
\ No newline at end of file
diff --git a/Engine/Source/Runtime/Core/Private/Android/AndroidJavaEnv.cpp b/Engine/Source/Runtime/Core/Private/Android/AndroidJavaEnv.cpp
--- a/Engine/Source/Runtime/Core/Private/Android/AndroidJavaEnv.cpp
+++ b/Engine/Source/Runtime/Core/Private/Android/AndroidJavaEnv.cpp
@@ -216,26 +216,26 @@ bool AndroidJavaEnv::CheckJavaException()
 
 FString FJavaHelper::FStringFromLocalRef(JNIEnv* Env, jstring JavaString)
 {
-	auto ReturnString = FStringFromParam(Env, JavaString);
+	FString ReturnString = FStringFromParam(Env, JavaString);
 	
 	if (Env && JavaString)
 	{
 		Env->DeleteLocalRef(JavaString);
 	}
 	
-	return MoveTemp(ReturnString);
+	return ReturnString;
 }
 
 FString FJavaHelper::FStringFromGlobalRef(JNIEnv* Env, jstring JavaString)
 {
-	auto ReturnString = FStringFromParam(Env, JavaString);
+	FString ReturnString = FStringFromParam(Env, JavaString);
 	
 	if (Env && JavaString)
 	{
 		Env->DeleteGlobalRef(JavaString);
 	}
 	
-	return MoveTemp(ReturnString);
+	return ReturnString;
 }
 
 FString FJavaHelper::FStringFromParam(JNIEnv* Env, jstring JavaString)
@@ -248,7 +248,7 @@ FString FJavaHelper::FStringFromParam(JNIEnv* Env, jstring JavaString)
 	const auto chars = Env->GetStringUTFChars(JavaString, 0);
 	FString ReturnString(UTF8_TO_TCHAR(chars));
 	Env->ReleaseStringUTFChars(JavaString, chars);
-	return MoveTemp(ReturnString);
+	return ReturnString;
 }
 
 FScopedJavaObject<jstring> FJavaHelper::ToJavaString(JNIEnv* Env, const FString& UnrealString)
diff --git a/Engine/Source/Runtime/Core/Private/Misc/EmbeddedCommunication.cpp b/Engine/Source/Runtime/Core/Private/Misc/EmbeddedCommunication.cpp
--- a/Engine/Source/Runtime/Core/Private/Misc/EmbeddedCommunication.cpp
+++ b/Engine/Source/Runtime/Core/Private/Misc/EmbeddedCommunication.cpp
@@ -201,6 +201,50 @@ void FEmbeddedCommunication::AllowSleep(FName Requester)
 #endif
 }
 
+#if BUILD_EMBEDDED_APP
+static FName NativeLogCategory("Native");
+#endif
+
+void FEmbeddedCommunication::UELogError(const TCHAR* String)
+{
+#if BUILD_EMBEDDED_APP
+	if (GWarn)
+	{
+		GWarn->Log(NativeLogCategory, ELogVerbosity::Error, String);
+	}
+#endif
+}
+
+void FEmbeddedCommunication::UELogWarning(const TCHAR* String)
+{
+#if BUILD_EMBEDDED_APP
+	if (GWarn)
+	{
+		GWarn->Log(NativeLogCategory, ELogVerbosity::Warning, String);
+	}
+#endif
+}
+
+void FEmbeddedCommunication::UELogLog(const TCHAR* String)
+{
+#if BUILD_EMBEDDED_APP
+	if (GLog)
+	{
+		GLog->Log(NativeLogCategory, ELogVerbosity::Log, String);
+	}
+#endif
+}
+
+void FEmbeddedCommunication::UELogVerbose(const TCHAR* String)
+{
+#if BUILD_EMBEDDED_APP
+	if (GLog)
+	{
+		GLog->Log(NativeLogCategory, ELogVerbosity::Verbose, String);
+	}
+#endif
+}
+
 bool FEmbeddedCommunication::IsAwakeForTicking()
 {
 #if BUILD_EMBEDDED_APP
diff --git a/Engine/Source/Runtime/Core/Public/Misc/EmbeddedCommunication.h b/Engine/Source/Runtime/Core/Public/Misc/EmbeddedCommunication.h
--- a/Engine/Source/Runtime/Core/Public/Misc/EmbeddedCommunication.h
+++ b/Engine/Source/Runtime/Core/Public/Misc/EmbeddedCommunication.h
@@ -81,6 +81,11 @@ public:
 	static void KeepAwake(FName Requester, bool bNeedsRendering);
 	static void AllowSleep(FName Requester);
 	
+	static void UELogError(const TCHAR* String);
+	static void UELogWarning(const TCHAR* String);
+	static void UELogLog(const TCHAR* String);
+	static void UELogVerbose(const TCHAR* String);
+
 	static bool IsAwakeForTicking();
 	static bool IsAwakeForRendering();
 	
diff --git a/Engine/Source/Runtime/Launch/Private/Android/AndroidJNI.cpp b/Engine/Source/Runtime/Launch/Private/Android/AndroidJNI.cpp
--- a/Engine/Source/Runtime/Launch/Private/Android/AndroidJNI.cpp
+++ b/Engine/Source/Runtime/Launch/Private/Android/AndroidJNI.cpp
@@ -1800,6 +1800,32 @@ JNI_METHOD void Java_com_epicgames_ue4_NativeCalls_AllowSleep(JNIEnv* jenv, jobj
 	FEmbeddedCommunication::AllowSleep(*Requester);
 }
 
+JNI_METHOD void Java_com_epicgames_ue4_NativeCalls_UELogError(JNIEnv* jenv, jobject thiz, jstring InString)
+{
+	const auto chars = jenv->GetStringUTFChars(InString, 0);
+	FEmbeddedCommunication::UELogError(UTF8_TO_TCHAR(chars));
+	jenv->ReleaseStringUTFChars(InString, chars);
+}
+
+JNI_METHOD void Java_com_epicgames_ue4_NativeCalls_UELogWarning(JNIEnv* jenv, jobject thiz, jstring InString)
+{
+	const auto chars = jenv->GetStringUTFChars(InString, 0);
+	FEmbeddedCommunication::UELogWarning(UTF8_TO_TCHAR(chars));
+	jenv->ReleaseStringUTFChars(InString, chars);
+}
+JNI_METHOD void Java_com_epicgames_ue4_NativeCalls_UELogLog(JNIEnv* jenv, jobject thiz, jstring InString)
+{
+	const auto chars = jenv->GetStringUTFChars(InString, 0);
+	FEmbeddedCommunication::UELogLog(UTF8_TO_TCHAR(chars));
+	jenv->ReleaseStringUTFChars(InString, chars);
+}
+
+JNI_METHOD void Java_com_epicgames_ue4_NativeCalls_UELogVerbose(JNIEnv* jenv, jobject thiz, jstring InString)
+{
+	const auto chars = jenv->GetStringUTFChars(InString, 0);
+	FEmbeddedCommunication::UELogVerbose(UTF8_TO_TCHAR(chars));
+	jenv->ReleaseStringUTFChars(InString, chars);
+}
 
 void FJavaWrapper::SetupEmbeddedCommunication(JNIEnv* Env)
 {
