- Added an optional funciton that can be called really early in embedded lifetime so that a game can initialize code before it may normally be ready, when using a global auto-register type class interacts poorly with other globals

diff --git a/Engine/Source/Runtime/Launch/Private/LaunchEngineLoop.cpp b/Engine/Source/Runtime/Launch/Private/LaunchEngineLoop.cpp
--- a/Engine/Source/Runtime/Launch/Private/LaunchEngineLoop.cpp
+++ b/Engine/Source/Runtime/Launch/Private/LaunchEngineLoop.cpp
@@ -1163,6 +1163,10 @@ int32 FEngineLoop::PreInit(const TCHAR* CmdLine)
 #endif // PLATFORM_WINDOWS
 
 #if BUILD_EMBEDDED_APP
+#ifdef EMBEDDED_LINKER_GAME_HELPER_FUNCTION
+	extern void EMBEDDED_LINKER_GAME_HELPER_FUNCTION();
+	EMBEDDED_LINKER_GAME_HELPER_FUNCTION();
+#endif
 	FEmbeddedCommunication::Init();
 	FEmbeddedCommunication::KeepAwake(TEXT("Startup"), false);
 #endif
