UE-65424 Dismissing an alarm in a project will leave the project's sounds muted

#4.22
#iOS

Alarm interrupts trigger AVAudioSessionInterruptionNotification. The audio context needs to be restarted on AVAudioSessionInterruptionTypeEnded (unmuting the audio is not enough, the system automatically suspends it).

Note: Also tested:
 -Push Notifications (i.e. FB messages), Reminders, SMS messages are handled by the OS (do not use AVAudioSessionInterruptionNotification, nor applicationWillResignActive/applicationDidBecomeActive).
- Phone calls:
Receive: applicationWillResignActive -> AVAudioSessionInterruptionNotification(AVAudioSessionInterruptionTypeBegan)
Close/reject: AVAudioSessionInterruptionNotification(AVAudioSessionInterruptionTypeEnded) -> applicationDidBecomeActive
- Notification Center (swipe down): applicationWillResignActive/applicationDidBecomeActive

diff --git a/Engine/Source/Runtime/ApplicationCore/Private/IOS/IOSAppDelegate.cpp b/Engine/Source/Runtime/ApplicationCore/Private/IOS/IOSAppDelegate.cpp
--- a/Engine/Source/Runtime/ApplicationCore/Private/IOS/IOSAppDelegate.cpp
+++ b/Engine/Source/Runtime/ApplicationCore/Private/IOS/IOSAppDelegate.cpp
@@ -449,6 +449,13 @@ static IOSAppDelegate* CachedDelegate = nil;
 				break;
 
 			case AVAudioSessionInterruptionTypeEnded:
+
+				NSNumber * interruptionOption = [[notification userInfo] objectForKey:AVAudioSessionInterruptionOptionKey];
+				if (interruptionOption != nil && interruptionOption.unsignedIntegerValue > 0)
+				{
+					FAppEntry::RestartAudio();
+				}
+
 				FAppEntry::Resume(true);
 				[self ToggleAudioSession:true force:true];
 				break;
diff --git a/Engine/Source/Runtime/ApplicationCore/Public/IOS/IOSAppDelegate.h b/Engine/Source/Runtime/ApplicationCore/Public/IOS/IOSAppDelegate.h
--- a/Engine/Source/Runtime/ApplicationCore/Public/IOS/IOSAppDelegate.h
+++ b/Engine/Source/Runtime/ApplicationCore/Public/IOS/IOSAppDelegate.h
@@ -73,6 +73,8 @@ namespace FAppEntry
 	void Shutdown();
     void Suspend(bool bIsInterrupt = false);
     void Resume(bool bIsInterrupt = false);
+	void RestartAudio();
+
 	bool IsStartupMoviePlaying();
 
 	extern bool	gAppLaunchedWithLocalNotification;
diff --git a/Engine/Source/Runtime/Launch/Private/IOS/LaunchIOS.cpp b/Engine/Source/Runtime/Launch/Private/IOS/LaunchIOS.cpp
--- a/Engine/Source/Runtime/Launch/Private/IOS/LaunchIOS.cpp
+++ b/Engine/Source/Runtime/Launch/Private/IOS/LaunchIOS.cpp
@@ -177,6 +177,37 @@ void FAppEntry::Resume(bool bIsInterrupt)
 	}
 }
 
+void FAppEntry::RestartAudio()
+{
+	if (GEngine && GEngine->GetMainAudioDevice())
+	{
+		FAudioDevice* AudioDevice = GEngine->GetMainAudioDevice();
+
+		if (FTaskGraphInterface::IsRunning())
+		{
+			int32& SuspendCounter = FIOSAudioDevice::GetSuspendCounter();
+
+			//increment the counter, otherwise ResumeContext won't work
+			if (SuspendCounter == 0)
+			{
+				FPlatformAtomics::InterlockedIncrement(&SuspendCounter);
+			}
+
+			FFunctionGraphTask::CreateAndDispatchWhenReady([AudioDevice]()
+			{
+				FAudioThread::RunCommandOnAudioThread([AudioDevice]()
+				{
+					AudioDevice->ResumeContext();
+				}, TStatId());
+			}, TStatId(), NULL, ENamedThreads::GameThread);
+		}
+		else
+		{
+			AudioDevice->ResumeContext();
+		}
+	}
+}
+
 void FAppEntry::PreInit(IOSAppDelegate* AppDelegate, UIApplication* Application)
 {
 	// make a controller object
