UE-66039 Match3 Audio stutters when opening notiffication center

#4.22
#ios

On applicationWillResignActive (FAppEntry::Suspend(false)) set a 0.5 sec delay before restarting the audio on resume, since another Suspend event may occur when pulling down the notification center (Suspend-Resume-Suspend). When the second Suspend occurs, the delay is reset.

Phone calls, notifications, backgrounding the app will reset the delay before resuming the game, so the audio will restart immediately on applicationDidBecomeActive(FAppEntry::Suspend(false))

Only when pressing the multitasking window, without sending the app to background, the audio is restarted on FAppEntry::Tick.

diff --git a/Engine/Source/Runtime/ApplicationCore/Private/IOS/IOSAppDelegate.cpp b/Engine/Source/Runtime/ApplicationCore/Private/IOS/IOSAppDelegate.cpp
--- a/Engine/Source/Runtime/ApplicationCore/Private/IOS/IOSAppDelegate.cpp
+++ b/Engine/Source/Runtime/ApplicationCore/Private/IOS/IOSAppDelegate.cpp
@@ -441,6 +441,9 @@ static IOSAppDelegate* CachedDelegate = nil;
 {
 	[[NSNotificationCenter defaultCenter] addObserverForName:AVAudioSessionInterruptionNotification object:nil queue:nil usingBlock:^(NSNotification *notification)
 	{
+		// the audio context should resume immediately after interrupt, if suspended
+		FAppEntry::ResetAudioContextResumeTime();
+
 		switch ([[[notification userInfo] objectForKey:AVAudioSessionInterruptionTypeKey] unsignedIntegerValue])
 		{
 			case AVAudioSessionInterruptionTypeBegan:
@@ -1404,6 +1407,9 @@ FCriticalSection RenderSuspend;
 	FIOSAsyncTask* AsyncTask = [[FIOSAsyncTask alloc] init];
 	AsyncTask.GameThreadCallback = ^ bool(void)
 	{
+		// the audio context should resume immediately after interrupt, if suspended
+		FAppEntry::ResetAudioContextResumeTime();
+
 		FCoreDelegates::ApplicationWillEnterBackgroundDelegate.Broadcast();
 		return true;
 	};
@@ -1418,6 +1424,9 @@ FCriticalSection RenderSuspend;
 	FIOSAsyncTask* AsyncTask = [[FIOSAsyncTask alloc] init];
 	AsyncTask.GameThreadCallback = ^ bool(void)
 	{
+		// the audio context should resume immediately after interrupt, if suspended
+		FAppEntry::ResetAudioContextResumeTime();
+
 		FCoreDelegates::ApplicationHasEnteredForegroundDelegate.Broadcast();
 		return true;
 	};
diff --git a/Engine/Source/Runtime/ApplicationCore/Public/IOS/IOSAppDelegate.h b/Engine/Source/Runtime/ApplicationCore/Public/IOS/IOSAppDelegate.h
--- a/Engine/Source/Runtime/ApplicationCore/Public/IOS/IOSAppDelegate.h
+++ b/Engine/Source/Runtime/ApplicationCore/Public/IOS/IOSAppDelegate.h
@@ -70,6 +70,8 @@ namespace FAppEntry
 	void Init();
 	void Tick();
     void SuspendTick();
+	void ResumeAudioContext();
+	void ResetAudioContextResumeTime();
 	void Shutdown();
     void Suspend(bool bIsInterrupt = false);
     void Resume(bool bIsInterrupt = false);
diff --git a/Engine/Source/Runtime/Launch/Private/IOS/LaunchIOS.cpp b/Engine/Source/Runtime/Launch/Private/IOS/LaunchIOS.cpp
--- a/Engine/Source/Runtime/Launch/Private/IOS/LaunchIOS.cpp
+++ b/Engine/Source/Runtime/Launch/Private/IOS/LaunchIOS.cpp
@@ -40,6 +40,14 @@ FAutoConsoleVariableRef CVarDisableAudioSuspendOnAudioInterrupt(
     TEXT("0: Not Disabled, 1: Disabled"),
     ECVF_Default);
 
+static const double cMaxAudioContextResumeDelay = 0.5;    // Setting this to be 0.5 seconds
+static double AudioContextResumeTime = 0;
+
+void FAppEntry::ResetAudioContextResumeTime()
+{
+	AudioContextResumeTime = 0;
+}
+
 void FAppEntry::Suspend(bool bIsInterrupt)
 {
 	// also treats interrupts BEFORE initializing the engine
@@ -74,6 +82,18 @@ void FAppEntry::Suspend(bool bIsInterrupt)
 			}
 			else
 			{
+				if (AudioContextResumeTime == 0)
+				{
+					// wait 0.5 sec before restarting the audio on resume
+					// another Suspend event may occur when pulling down the notification center (Suspend-Resume-Suspend)
+					AudioContextResumeTime = FPlatformTime::Seconds() + cMaxAudioContextResumeDelay;
+				}
+				else
+				{
+					//second resume, restart the audio immediately after resume
+					AudioContextResumeTime = 0; 
+				}
+
 				if (FTaskGraphInterface::IsRunning())
 				{
 					FGraphEventRef ResignTask = FFunctionGraphTask::CreateAndDispatchWhenReady([AudioDevice]()
@@ -150,19 +170,15 @@ void FAppEntry::Resume(bool bIsInterrupt)
 			}
 			else
 			{
-				if (FTaskGraphInterface::IsRunning())
+				if (AudioContextResumeTime != 0)
 				{
-					FFunctionGraphTask::CreateAndDispatchWhenReady([AudioDevice]()
-					{
-						FAudioThread::RunCommandOnAudioThread([AudioDevice]()
-						{
-							AudioDevice->ResumeContext();
-						}, TStatId());
-					}, TStatId(), NULL, ENamedThreads::GameThread);
+					// resume audio on Tick()
+					AudioContextResumeTime = FPlatformTime::Seconds() + cMaxAudioContextResumeDelay;
 				}
 				else
 				{
-					AudioDevice->ResumeContext();
+					// resume audio immediately
+					ResumeAudioContext();
 				}
 			}
 		}
@@ -177,6 +193,32 @@ void FAppEntry::Resume(bool bIsInterrupt)
 	}
 }
 
+
+void FAppEntry::ResumeAudioContext()
+{
+	if (GEngine && GEngine->GetMainAudioDevice())
+	{
+		FAudioDevice* AudioDevice = GEngine->GetMainAudioDevice();
+		if (AudioDevice)
+		{
+			if (FTaskGraphInterface::IsRunning())
+			{
+				FFunctionGraphTask::CreateAndDispatchWhenReady([AudioDevice]()
+				{
+					FAudioThread::RunCommandOnAudioThread([AudioDevice]()
+					{
+						AudioDevice->ResumeContext();
+					}, TStatId());
+				}, TStatId(), NULL, ENamedThreads::GameThread);
+			}
+			else
+			{
+				AudioDevice->ResumeContext();
+			}
+		}
+	}
+}
+
 void FAppEntry::RestartAudio()
 {
 	if (GEngine && GEngine->GetMainAudioDevice())
@@ -380,6 +422,15 @@ void FAppEntry::Tick()
         FPlatformProcess::SetRealTimeMode();
     }
     
+	if (AudioContextResumeTime != 0)
+	{
+		if (FPlatformTime::Seconds() >= AudioContextResumeTime)
+		{
+			ResumeAudioContext();
+			AudioContextResumeTime = 0;
+		}
+	}
+
 	// tick the engine
 	GEngineLoop.Tick();
 }
