Remove access token logging of java bridge on Android for shipping builds

#android

diff --git a/Engine/Source/Runtime/Launch/Private/Android/AndroidJNI.cpp b/Engine/Source/Runtime/Launch/Private/Android/AndroidJNI.cpp
--- a/Engine/Source/Runtime/Launch/Private/Android/AndroidJNI.cpp
+++ b/Engine/Source/Runtime/Launch/Private/Android/AndroidJNI.cpp
@@ -1713,8 +1713,9 @@ void AndroidThunkCpp_OnNativeToEmbeddedReply(FString ID, const FEmbeddedCommunic
 	if (JNIEnv* Env = FAndroidApplication::GetJavaEnv())
 	{
 		// marshall some data
-
+#if !UE_BUILD_SHIPPING
 		UE_LOG(LogInit, Display, TEXT("Java call Id: %s, Routing Function: %s, Error %s, Params:"), *ID, *RoutingFunction, *InError);
+#endif
 
 		// create a string array for the pairs
 		static auto StringClass = FJavaWrapper::FindClassGlobalRef(Env, "java/lang/String", false);
@@ -1726,7 +1727,9 @@ void AndroidThunkCpp_OnNativeToEmbeddedReply(FString ID, const FEmbeddedCommunic
 			auto ValueString = FJavaHelper::ToJavaString(Env, It.Value);
 			Env->SetObjectArrayElement(*ReturnValues, Index++, *KeyString);
 			Env->SetObjectArrayElement(*ReturnValues, Index++, *ValueString);
+#if !UE_BUILD_SHIPPING
 			UE_LOG(LogInit, Display, TEXT("  %s : %s"), *It.Key, *It.Value);
+#endif
 		}
 		auto Error = FJavaHelper::ToJavaString(Env, InError);
 		
@@ -1757,7 +1760,9 @@ JNI_METHOD void Java_com_epicgames_ue4_NativeCalls_CallNativeToEmbedded(JNIEnv*
 	// wake up UE
 	FEmbeddedCommunication::KeepAwake(CommandName, false);
 
+#if !UE_BUILD_SHIPPING
 	FPlatformMisc::LowLevelOutputDebugStringf(TEXT("NativeToEmbeddedCall: Subsystem: %s, Command: %s, Params:"), *Subsystem, *Helper.Command);
+#endif
 
 	if (InParams != nullptr)
 	{
@@ -1769,7 +1774,9 @@ JNI_METHOD void Java_com_epicgames_ue4_NativeCalls_CallNativeToEmbedded(JNIEnv*
 			auto javaValue = FJavaHelper::FStringFromLocalRef(jenv, (jstring)(jenv->GetObjectArrayElement(InParams, Index++)));
 
 			Helper.Parameters.Add(javaKey, javaValue);
+#if !UE_BUILD_SHIPPING
 			FPlatformMisc::LowLevelOutputDebugStringf(TEXT("  %s : %s"), *javaKey, *javaValue);
+#endif
 		}
 	}
 	
