Change "-tracehost" to use the common convention for arguments with parameters.

(ushell-p4-cherrypick of 6836373 by [REDACTED_PERSON])

diff --git a/Engine/Source/Runtime/Launch/Private/LaunchEngineLoop.cpp b/Engine/Source/Runtime/Launch/Private/LaunchEngineLoop.cpp
--- a/Engine/Source/Runtime/Launch/Private/LaunchEngineLoop.cpp
+++ b/Engine/Source/Runtime/Launch/Private/LaunchEngineLoop.cpp
@@ -1215,7 +1215,7 @@ int32 FEngineLoop::PreInit(const TCHAR* CmdLine)
 
 	{
 		FString TraceHost;
-		if (FParse::Value(CmdLine, TEXT("tracehost"), TraceHost))
+		if (FParse::Value(CmdLine, TEXT("-tracehost="), TraceHost))
 		{
 			Trace::Connect(*TraceHost);
 		}
