Copied cl 6853762from Dev-Editor

Fix ShaderCompileWorker crashing when setting FORCE_USE_STATS

ShaderCompileWorker was passing "-ReduceThreadUsage" command line and FTaskGraphInterface::Startup() wasn't getting called so FTaskGraphInterface::Get() was returning NULL and crashing ShaderCompileWorker processes.

diff --git a/Engine/Source/Runtime/Launch/Private/LaunchEngineLoop.cpp b/Engine/Source/Runtime/Launch/Private/LaunchEngineLoop.cpp
--- a/Engine/Source/Runtime/Launch/Private/LaunchEngineLoop.cpp
+++ b/Engine/Source/Runtime/Launch/Private/LaunchEngineLoop.cpp
@@ -1699,7 +1699,8 @@ int32 FEngineLoop::PreInit(const TCHAR* CmdLine)
 
 	// Some programs might not use the taskgraph or thread pool
 	bool bCreateTaskGraphAndThreadPools = true;
-#if IS_PROGRAM
+	// If STATS is defined (via FORCE_USE_STATS or other), we have to call FTaskGraphInterface::Startup()
+#if IS_PROGRAM && !STATS
 	bCreateTaskGraphAndThreadPools = !FParse::Param(FCommandLine::Get(), TEXT("ReduceThreadUsage"));
 #endif
 	if (bCreateTaskGraphAndThreadPools)
