Wait for audio thread to finish on Android when backgrounding.
Fixes an instance of GC timeout which is waiting for USoundWave objects.

diff --git a/Engine/Source/Runtime/Launch/Private/Android/AndroidEventManager.cpp b/Engine/Source/Runtime/Launch/Private/Android/AndroidEventManager.cpp
--- a/Engine/Source/Runtime/Launch/Private/Android/AndroidEventManager.cpp
+++ b/Engine/Source/Runtime/Launch/Private/Android/AndroidEventManager.cpp
@@ -324,6 +324,10 @@ void FAppEventManager::PauseAudio()
 	{
 		if (AudioDevice->IsAudioMixerEnabled())
 		{
+			FAudioCommandFence Fence;
+			Fence.BeginFence();
+			Fence.Wait();
+
 			AudioDevice->SuspendContext();
 		}
 		else
