Avoid overriding main() for the platform group.

diff --git a/Engine/Source/Runtime/Launch/Private/Linux/LaunchLinux.cpp b/Engine/Source/Runtime/Launch/Private/Linux/LaunchLinux.cpp
new file mode 100644
--- /dev/null
+++ b/Engine/Source/Runtime/Launch/Private/Linux/LaunchLinux.cpp
@@ -0,0 +1,17 @@
+// Copyright 1998-2019 Epic Games, Inc. All Rights Reserved.
+
+#include "CoreMinimal.h"
+#include "LaunchEngineLoop.h"
+#include "UnixCommonStartup.h"
+
+extern int32 GuardedMain( const TCHAR* CmdLine );
+
+/**
+ * Workaround function to avoid circular dependencies between Launch and CommonUnixStartup modules (see LaunchUnix.cpp)
+ */
+void LaunchUnix_FEngineLoop_AppExit();
+
+int main(int argc, char *argv[])
+{
+	return CommonUnixMain(argc, argv, &GuardedMain, &LaunchUnix_FEngineLoop_AppExit);
+}
diff --git a/Engine/Source/Runtime/Launch/Private/Unix/LaunchUnix.cpp b/Engine/Source/Runtime/Launch/Private/Unix/LaunchUnix.cpp
--- a/Engine/Source/Runtime/Launch/Private/Unix/LaunchUnix.cpp
+++ b/Engine/Source/Runtime/Launch/Private/Unix/LaunchUnix.cpp
@@ -17,8 +17,3 @@ void LaunchUnix_FEngineLoop_AppExit()
 {
 	return FEngineLoop::AppExit();
 }
-
-int main(int argc, char *argv[])
-{
-	return CommonUnixMain(argc, argv, &GuardedMain, &LaunchUnix_FEngineLoop_AppExit);
-}
