Fix conflation of dumping early reads with config patching.

diff --git a/Engine/Source/Runtime/Launch/Private/LaunchEngineLoop.cpp b/Engine/Source/Runtime/Launch/Private/LaunchEngineLoop.cpp
--- a/Engine/Source/Runtime/Launch/Private/LaunchEngineLoop.cpp
+++ b/Engine/Source/Runtime/Launch/Private/LaunchEngineLoop.cpp
@@ -892,9 +892,10 @@ void LaunchUpdateMostRecentProjectFile()
 }
 
 #if WITH_ENGINE
-void OnStartupContentMounted(FInstallBundleResultInfo Result, bool bDumpEarlyConfigReads, bool bDumpEarlyPakFileReads);
+void OnStartupContentMounted(FInstallBundleResultInfo Result, bool bDumpEarlyConfigReads, bool bDumpEarlyPakFileReads, bool bReloadConfig);
 #endif
-void HandleConfigReload(bool bDumpEarlyConfigReads, bool bDumpEarlyPakFileReads);
+void DumpEarlyReads(bool bDumpEarlyConfigReads, bool bDumpEarlyPakFileReads);
+void HandleConfigReload(bool bReloadConfig);
 
 #if !UE_BUILD_SHIPPING
 class FFileInPakFileHistoryHelper
@@ -1743,18 +1744,19 @@ int32 FEngineLoop::PreInit(const TCHAR* CmdLine)
 	constexpr bool bWithConfigPatching = false;
 #endif
 
-	if (bWithConfigPatching)
+	if (bDumpEarlyConfigReads)
 	{
-		UE_LOG(LogInit, Verbose, TEXT("Begin recording CVar changes for config patching."));
+		RecordConfigReadsFromIni();
+	}
 
-		if (bDumpEarlyConfigReads)
-		{
-			RecordConfigReadsFromIni();
-		}
-		if (bDumpEarlyPakFileReads)
-		{
-			RecordFileReadsFromPaks();
-		}
+	if (bDumpEarlyPakFileReads)
+	{
+		RecordFileReadsFromPaks();
+	}
+
+	if(bWithConfigPatching)
+	{
+		UE_LOG(LogInit, Verbose, TEXT("Begin recording CVar changes for config patching."));
 
 		RecordApplyCVarSettingsFromIni();
 	}
@@ -2387,16 +2389,13 @@ int32 FEngineLoop::PreInit(const TCHAR* CmdLine)
 					return 1;
 				}
 			}
-
-			if (bWithConfigPatching)
+			
+			IInstallBundleManager* BundleManager = IInstallBundleManager::GetPlatformInstallBundleManager();
+			if (BundleManager != nullptr && !BundleManager->IsNullInterface())
 			{
-				IInstallBundleManager* BundleManager = IInstallBundleManager::GetPlatformInstallBundleManager();
-				if (BundleManager != nullptr && !BundleManager->IsNullInterface())
-				{
-					IInstallBundleManager::InstallBundleCompleteDelegate.AddStatic(&OnStartupContentMounted, bDumpEarlyConfigReads, bDumpEarlyPakFileReads);
-				}
-				// If not using the bundle manager, config will be reloaded after ESP, see below
+				IInstallBundleManager::InstallBundleCompleteDelegate.AddStatic(&OnStartupContentMounted, bDumpEarlyConfigReads, bDumpEarlyPakFileReads, bWithConfigPatching);
 			}
+			// If not using the bundle manager, config will be reloaded after ESP, see below
 
 			if (GetMoviePlayer()->HasEarlyStartupMovie())
 			{
@@ -2526,12 +2525,13 @@ int32 FEngineLoop::PreInit(const TCHAR* CmdLine)
 				FCoreDelegates::OnMountAllPakFiles.Execute(PakFolders);
 			}
 
+			DumpEarlyReads(bDumpEarlyConfigReads, bDumpEarlyPakFileReads);
+
 			//Reapply CVars after our EarlyLoadScreen
 			if(bWithConfigPatching)
 			{
 				SCOPED_BOOT_TIMING("ReapplyCVarsFromIniAfterEarlyStartupScreen");
-
-				HandleConfigReload(bDumpEarlyConfigReads, bDumpEarlyPakFileReads);
+				HandleConfigReload(bWithConfigPatching);
 			}
 
 			//Handle opening shader library after our EarlyLoadScreen
@@ -3787,18 +3787,19 @@ void FEngineLoop::ProcessLocalPlayerSlateOperations() const
 }
 
 #if WITH_ENGINE
-void OnStartupContentMounted(FInstallBundleResultInfo Result, bool bDumpEarlyConfigReads, bool bDumpEarlyPakFileReads)
+void OnStartupContentMounted(FInstallBundleResultInfo Result, bool bDumpEarlyConfigReads, bool bDumpEarlyPakFileReads, bool bReloadConfig)
 {
 	if (Result.bIsStartup && Result.Result == EInstallBundleResult::OK)
 	{
-		HandleConfigReload(bDumpEarlyConfigReads, bDumpEarlyPakFileReads);
+		DumpEarlyReads(bDumpEarlyConfigReads, bDumpEarlyPakFileReads);
+		HandleConfigReload(bReloadConfig);
 
 		IInstallBundleManager::InstallBundleCompleteDelegate.RemoveAll(&GEngineLoop);
 	}
 }
 #endif
 
-void HandleConfigReload(bool bDumpEarlyConfigReads, bool bDumpEarlyPakFileReads)
+void DumpEarlyReads(bool bDumpEarlyConfigReads, bool bDumpEarlyPakFileReads)
 {
 	if (bDumpEarlyConfigReads)
 	{
@@ -3811,9 +3812,15 @@ void HandleConfigReload(bool bDumpEarlyConfigReads, bool bDumpEarlyPakFileReads)
 		DumpRecordedFileReadsFromPaks();
 		DeleteRecordedFileReadsFromPaks();
 	}
+}
 
-	ReapplyRecordedCVarSettingsFromIni();
-	DeleteRecordedCVarSettingsFromIni();
+void HandleConfigReload(bool bReloadConfig)
+{
+	if (bReloadConfig)
+	{
+		ReapplyRecordedCVarSettingsFromIni();
+		DeleteRecordedCVarSettingsFromIni();
+	}
 }
 
 bool FEngineLoop::ShouldUseIdleMode() const
