Automatically re-launch the correct executable if using a project with a unique build environment. The .target file is read at startup, and used to determine the correct path for the executable to launch.

diff --git a/Engine/Source/Developer/DesktopPlatform/Private/DesktopPlatformBase.cpp b/Engine/Source/Developer/DesktopPlatform/Private/DesktopPlatformBase.cpp
--- a/Engine/Source/Developer/DesktopPlatform/Private/DesktopPlatformBase.cpp
+++ b/Engine/Source/Developer/DesktopPlatform/Private/DesktopPlatformBase.cpp
@@ -562,20 +562,21 @@ bool FDesktopPlatformBase::CompileGameProject(const FString& RootDir, const FStr
 {
 	FModuleManager& ModuleManager = FModuleManager::Get();
 
-	// Get the project directory
-	FString ProjectDir = FPaths::GetPath(ProjectFileName);
-
 	// Build the argument list
 	FString Arguments = FString::Printf(TEXT("%s %s"), ModuleManager.GetUBTConfiguration(), FPlatformMisc::GetUBTPlatform());
 
-	// Append the project name if it's a foreign project
-	if ( !ProjectFileName.IsEmpty() )
+	// Append the project name if it's a foreign project. Otherwise compile UE4Editor.
+	if ( ProjectFileName.IsEmpty() )
+	{
+		Arguments = TEXT("UE4Editor ") + Arguments;
+	}
+	else
 	{
-		Arguments += FString::Printf(TEXT(" -Project=\"%s\""), *IFileManager::Get().ConvertToAbsolutePathForExternalAppForRead(*ProjectFileName));
+		Arguments += FString::Printf(TEXT(" -Project=\"%s\" -TargetType=Editor"), *IFileManager::Get().ConvertToAbsolutePathForExternalAppForRead(*ProjectFileName));
 	}
 
 	// Append any other options
-	Arguments += " -TargetType=Editor -Progress -NoHotReloadFromIDE";
+	Arguments += " -Progress -NoHotReloadFromIDE";
 
 	// Run UBT
 	bool bResult = RunUnrealBuildTool(LOCTEXT("CompilingProject", "Compiling project..."), RootDir, Arguments, Warn);
diff --git a/Engine/Source/Developer/DesktopPlatform/Private/TargetReceipt.cpp b/Engine/Source/Developer/DesktopPlatform/Private/TargetReceipt.cpp
--- a/Engine/Source/Developer/DesktopPlatform/Private/TargetReceipt.cpp
+++ b/Engine/Source/Developer/DesktopPlatform/Private/TargetReceipt.cpp
@@ -141,6 +141,24 @@ bool FTargetReceipt::Read(const FString& FileName)
 	return true;
 }
 
+FString FTargetReceipt::GetDefaultPath(const TCHAR* BaseDir, const TCHAR* TargetName, const TCHAR* Platform, EBuildConfiguration Configuration, const TCHAR* BuildArchitecture)
+{
+	const TCHAR* ArchitectureSuffix = TEXT("");
+	if (BuildArchitecture != nullptr)
+	{
+		ArchitectureSuffix = BuildArchitecture;
+	}
+
+	if ((BuildArchitecture == nullptr || BuildArchitecture[0] == 0) && Configuration == EBuildConfiguration::Development)
+	{
+		return FPaths::Combine(BaseDir, FString::Printf(TEXT("Binaries/%s/%s.target"), Platform, TargetName));
+	}
+	else
+	{
+		return FPaths::Combine(BaseDir, FString::Printf(TEXT("Binaries/%s/%s-%s-%s%s.target"), Platform, TargetName, Platform, Configuration, ArchitectureSuffix));
+	}
+}
+
 void FTargetReceipt::ExpandVariables(FString& Path)
 {
 	static FString EngineDirPrefix = TEXT("$(EngineDir)");
diff --git a/Engine/Source/Developer/DesktopPlatform/Public/TargetReceipt.h b/Engine/Source/Developer/DesktopPlatform/Public/TargetReceipt.h
--- a/Engine/Source/Developer/DesktopPlatform/Public/TargetReceipt.h
+++ b/Engine/Source/Developer/DesktopPlatform/Public/TargetReceipt.h
@@ -87,6 +87,18 @@ struct DESKTOPPLATFORM_API FTargetReceipt
 	 */
 	bool Read(const FString& FileName);
 
+	/**
+	 * Gets the default path for a target receipt
+	 *
+	 * @param BaseDir Base directory for the target being built; either the project directory or engine directory.
+	 * @param TargetName The target being built
+	 * @param Platform The target platform
+	 * @param Configuration The target configuration
+	 * @param BuildArchitecture The architecture being built
+	 * @return Path to the receipt for this target
+	 */
+	static FString GetDefaultPath(const TCHAR* BaseDir, const TCHAR* TargetName, const TCHAR* Platform, EBuildConfiguration Configuration, const TCHAR* BuildArchitecture);
+
 private:
 	/**
 	 * Expands the $(EngineDir) and $(ProjectDir) variables within a string
diff --git a/Engine/Source/Editor/UnrealEd/Private/UnrealEdMisc.cpp b/Engine/Source/Editor/UnrealEd/Private/UnrealEdMisc.cpp
--- a/Engine/Source/Editor/UnrealEd/Private/UnrealEdMisc.cpp
+++ b/Engine/Source/Editor/UnrealEd/Private/UnrealEdMisc.cpp
@@ -836,16 +836,11 @@ bool FUnrealEdMisc::EnableWorldComposition(UWorld* InWorld, bool bEnable)
 	return true;
 }
 
-
 /** Build and return the path to the current project (used for relaunching the editor.)	 */
 FString CreateProjectPath()
 {
 #if PLATFORM_WINDOWS
-	// If we are running in 64 bit, launch the 64 bit process
-	const TCHAR* PlatformConfig = FPlatformMisc::GetUBTPlatform();
-	// Executable filename does not depend on the selected project. Simply create full path to the current executable.
-	FString ExeFileName = FPaths::EngineDir() / TEXT("Binaries") / PlatformConfig / FString( FPlatformProcess::ExecutableName() ) + TEXT(".exe");
-	return ExeFileName;
+	return FPlatformProcess::ExecutablePath();
 #elif PLATFORM_MAC
 	@autoreleasepool
 	{
diff --git a/Engine/Source/Runtime/Core/Private/Windows/WindowsPlatformProcess.cpp b/Engine/Source/Runtime/Core/Private/Windows/WindowsPlatformProcess.cpp
--- a/Engine/Source/Runtime/Core/Private/Windows/WindowsPlatformProcess.cpp
+++ b/Engine/Source/Runtime/Core/Private/Windows/WindowsPlatformProcess.cpp
@@ -725,7 +725,7 @@ bool FWindowsPlatformProcess::ExecProcess(const TCHAR* URL, const TCHAR* Params,
 
 	bool bSuccess = false;
 
-	FString CommandLine = FString::Printf(TEXT("%s %s"), URL, Params);
+	FString CommandLine = FString::Printf(TEXT("\"%s\" %s"), URL, Params);
 
 	PROCESS_INFORMATION ProcInfo;
 	if (CreateProcess(NULL, CommandLine.GetCharArray().GetData(), NULL, NULL, TRUE, NORMAL_PRIORITY_CLASS | DETACHED_PROCESS | EXTENDED_STARTUPINFO_PRESENT, NULL, OptionalWorkingDirectory, &StartupInfoEx.StartupInfo, &ProcInfo))
diff --git a/Engine/Source/Runtime/Launch/Private/LaunchEngineLoop.cpp b/Engine/Source/Runtime/Launch/Private/LaunchEngineLoop.cpp
--- a/Engine/Source/Runtime/Launch/Private/LaunchEngineLoop.cpp
+++ b/Engine/Source/Runtime/Launch/Private/LaunchEngineLoop.cpp
@@ -183,6 +183,7 @@ class FFeedbackContext;
 
 #if UE_EDITOR
 	#include "DesktopPlatformModule.h"
+	#include "TargetReceipt.h"
 #endif
 
 #define LOCTEXT_NAMESPACE "LaunchEngineLoop"
@@ -4652,6 +4653,28 @@ static void CheckForPrintTimesOverride()
 	}
 }
 
+#if UE_EDITOR
+bool LaunchCorrectEditorExecutable(const FString& EditorTargetFileName)
+{
+	FTargetReceipt Receipt;
+	if(FPaths::FileExists(EditorTargetFileName) && Receipt.Read(EditorTargetFileName))
+	{
+		FString CurrentExecutableName = FPlatformProcess::ExecutablePath();
+		FPaths::MakeStandardFilename(CurrentExecutableName);
+
+		FString LaunchExecutableName = Receipt.Launch;
+		FPaths::MakeStandardFilename(LaunchExecutableName);
+
+		if(LaunchExecutableName != CurrentExecutableName)
+		{
+			UE_LOG(LogInit, Display, TEXT("Running incorrect executable for target. Launching %s..."), *LaunchExecutableName);
+			FPlatformProcess::CreateProc(*IFileManager::Get().ConvertToAbsolutePathForExternalAppForRead(*LaunchExecutableName), FCommandLine::GetOriginal(), true, false, false, nullptr, 0, nullptr, nullptr, nullptr);
+			return true;
+		}
+	}
+	return false;
+}
+#endif
 
 /* FEngineLoop static interface
  *****************************************************************************/
@@ -4818,8 +4841,41 @@ bool FEngineLoop::AppInit( )
 
 	// Check whether the project or any of its plugins are missing or are out of date
 #if UE_EDITOR && !IS_MONOLITHIC
-	if(!GIsBuildMachine && FPaths::IsProjectFilePathSet() && PluginManager.AreRequiredPluginsAvailable())
+	if(!GIsBuildMachine && FPaths::IsProjectFilePathSet())
 	{
+		// Check all the plugins are present
+		if(!PluginManager.AreRequiredPluginsAvailable())
+		{
+			return false;
+		}
+
+		// Find the editor target
+		FString EditorTargetFileName;
+		for (const FTargetInfo& Target : FDesktopPlatformModule::Get()->GetTargetsForProject(FPaths::GetProjectFilePath()))
+		{
+			if (Target.Type == EBuildTargetType::Editor)
+			{
+				// Read the editor target receipt
+				EditorTargetFileName = FTargetReceipt::GetDefaultPath(FPlatformMisc::ProjectDir(), *Target.Name, FPlatformProcess::GetBinariesSubdirectory(), FApp::GetBuildConfiguration(), nullptr);
+				break;
+			}
+		}
+
+		// If there was no editor target, use the default UE4Editor target
+		FString CompileProject = FPaths::GetProjectFilePath();
+		if(EditorTargetFileName.Len() == 0)
+		{
+			CompileProject.Empty();
+			EditorTargetFileName = FTargetReceipt::GetDefaultPath(FPlatformMisc::EngineDir(), TEXT("UE4Editor"), FPlatformProcess::GetBinariesSubdirectory(), FApp::GetBuildConfiguration(), nullptr);
+		}
+
+		// If we're not running the correct executable for the current target, and the listed executable exists, run that instead
+		if(LaunchCorrectEditorExecutable(EditorTargetFileName))
+		{
+			return false;
+		}
+
+		// Check if we need to compile
 		bool bNeedCompile = false;
 		GConfig->GetBool(TEXT("/Script/UnrealEd.EditorLoadingSavingSettings"), TEXT("bForceCompilationAtStartup"), bNeedCompile, GEditorPerProjectIni);
 		if(FParse::Param(FCommandLine::Get(), TEXT("SKIPCOMPILE")) || FParse::Param(FCommandLine::Get(), TEXT("MULTIPROCESS")))
@@ -4870,6 +4926,14 @@ bool FEngineLoop::AppInit( )
 
 				bNeedCompile = true;
 			}
+			else if(!FPaths::FileExists(EditorTargetFileName))
+			{
+				// Prompt to compile. The target file isn't essential, but we 
+				if (FPlatformMisc::MessageBoxExt(EAppMsgType::YesNo, *FString::Printf(TEXT("The %s file does not exist. Would you like to build the editor?"), *FPaths::GetCleanFilename(EditorTargetFileName)), TEXT("Missing target file")) == EAppReturnType::Yes)
+				{
+					bNeedCompile = true;
+				}
+			}
 		}
 
 		FEmbeddedCommunication::ForceTick(16);
@@ -4879,9 +4943,15 @@ bool FEngineLoop::AppInit( )
 			// Try to compile it
 			FFeedbackContext *Context = (FFeedbackContext*)FDesktopPlatformModule::Get()->GetNativeFeedbackContext();
 			Context->BeginSlowTask(FText::FromString(TEXT("Starting build...")), true, true);
-			bool bCompileResult = FDesktopPlatformModule::Get()->CompileGameProject(FPaths::RootDir(), FPaths::GetProjectFilePath(), Context);
+			bool bCompileResult = FDesktopPlatformModule::Get()->CompileGameProject(FPaths::RootDir(), CompileProject, Context);
 			Context->EndSlowTask();
 
+			// Check if we're running the wrong executable now
+			if(bCompileResult && LaunchCorrectEditorExecutable(EditorTargetFileName))
+			{
+				return false;
+			}
+
 			// Get a list of modules which are still incompatible
 			TArray<FString> StillIncompatibleFiles;
 			ProjectManager.CheckModuleCompatibility(StillIncompatibleFiles);
