Prevent launching the wrong executable when running in unattended mode.

#rnx

diff --git a/Engine/Source/Runtime/Launch/Private/LaunchEngineLoop.cpp b/Engine/Source/Runtime/Launch/Private/LaunchEngineLoop.cpp
--- a/Engine/Source/Runtime/Launch/Private/LaunchEngineLoop.cpp
+++ b/Engine/Source/Runtime/Launch/Private/LaunchEngineLoop.cpp
@@ -4657,7 +4657,7 @@ static void CheckForPrintTimesOverride()
 bool LaunchCorrectEditorExecutable(const FString& EditorTargetFileName)
 {
 	FTargetReceipt Receipt;
-	if(FPaths::FileExists(EditorTargetFileName) && Receipt.Read(EditorTargetFileName))
+	if(!FApp::IsUnattended() && FPaths::FileExists(EditorTargetFileName) && Receipt.Read(EditorTargetFileName))
 	{
 		FString CurrentExecutableName = FPlatformProcess::ExecutablePath();
 		FPaths::MakeStandardFilename(CurrentExecutableName);
