Add -ForceQuitAfterEarlyReads to exit as soon as we are done recording early config/pak reads.

diff --git a/Engine/Source/Runtime/Launch/Private/LaunchEngineLoop.cpp b/Engine/Source/Runtime/Launch/Private/LaunchEngineLoop.cpp
--- a/Engine/Source/Runtime/Launch/Private/LaunchEngineLoop.cpp
+++ b/Engine/Source/Runtime/Launch/Private/LaunchEngineLoop.cpp
@@ -892,9 +892,9 @@ void LaunchUpdateMostRecentProjectFile()
 }
 
 #if WITH_ENGINE
-void OnStartupContentMounted(FInstallBundleResultInfo Result, bool bDumpEarlyConfigReads, bool bDumpEarlyPakFileReads, bool bReloadConfig);
+void OnStartupContentMounted(FInstallBundleResultInfo Result, bool bDumpEarlyConfigReads, bool bDumpEarlyPakFileReads, bool bReloadConfig, bool bForceQuitAfterEarlyReads);
 #endif
-void DumpEarlyReads(bool bDumpEarlyConfigReads, bool bDumpEarlyPakFileReads);
+void DumpEarlyReads(bool bDumpEarlyConfigReads, bool bDumpEarlyPakFileReads, bool bForceQuitAfterEarlyReads);
 void HandleConfigReload(bool bReloadConfig);
 
 #if !UE_BUILD_SHIPPING
@@ -1736,6 +1736,7 @@ int32 FEngineLoop::PreInit(const TCHAR* CmdLine)
 
 	const bool bDumpEarlyConfigReads = FParse::Param(FCommandLine::Get(), TEXT("DumpEarlyConfigReads"));
 	const bool bDumpEarlyPakFileReads = FParse::Param(FCommandLine::Get(), TEXT("DumpEarlyPakFileReads"));
+	const bool bForceQuitAfterEarlyReads = FParse::Param(FCommandLine::Get(), TEXT("ForceQuitAfterEarlyReads"));
 
 	// Overly verbose to avoid a dumb static analysis warning
 #if WITH_CONFIG_PATCHING
@@ -2393,7 +2394,8 @@ int32 FEngineLoop::PreInit(const TCHAR* CmdLine)
 			IInstallBundleManager* BundleManager = IInstallBundleManager::GetPlatformInstallBundleManager();
 			if (BundleManager != nullptr && !BundleManager->IsNullInterface())
 			{
-				IInstallBundleManager::InstallBundleCompleteDelegate.AddStatic(&OnStartupContentMounted, bDumpEarlyConfigReads, bDumpEarlyPakFileReads, bWithConfigPatching);
+				IInstallBundleManager::InstallBundleCompleteDelegate.AddStatic(
+					&OnStartupContentMounted, bDumpEarlyConfigReads, bDumpEarlyPakFileReads, bWithConfigPatching, bForceQuitAfterEarlyReads);
 			}
 			// If not using the bundle manager, config will be reloaded after ESP, see below
 
@@ -2525,7 +2527,7 @@ int32 FEngineLoop::PreInit(const TCHAR* CmdLine)
 				FCoreDelegates::OnMountAllPakFiles.Execute(PakFolders);
 			}
 
-			DumpEarlyReads(bDumpEarlyConfigReads, bDumpEarlyPakFileReads);
+			DumpEarlyReads(bDumpEarlyConfigReads, bDumpEarlyPakFileReads, bForceQuitAfterEarlyReads);
 
 			//Reapply CVars after our EarlyLoadScreen
 			if(bWithConfigPatching)
@@ -3787,11 +3789,11 @@ void FEngineLoop::ProcessLocalPlayerSlateOperations() const
 }
 
 #if WITH_ENGINE
-void OnStartupContentMounted(FInstallBundleResultInfo Result, bool bDumpEarlyConfigReads, bool bDumpEarlyPakFileReads, bool bReloadConfig)
+void OnStartupContentMounted(FInstallBundleResultInfo Result, bool bDumpEarlyConfigReads, bool bDumpEarlyPakFileReads, bool bReloadConfig, bool bForceQuitAfterEarlyReads)
 {
 	if (Result.bIsStartup && Result.Result == EInstallBundleResult::OK)
 	{
-		DumpEarlyReads(bDumpEarlyConfigReads, bDumpEarlyPakFileReads);
+		DumpEarlyReads(bDumpEarlyConfigReads, bDumpEarlyPakFileReads, bForceQuitAfterEarlyReads);
 		HandleConfigReload(bReloadConfig);
 
 		IInstallBundleManager::InstallBundleCompleteDelegate.RemoveAll(&GEngineLoop);
@@ -3799,7 +3801,7 @@ void OnStartupContentMounted(FInstallBundleResultInfo Result, bool bDumpEarlyCon
 }
 #endif
 
-void DumpEarlyReads(bool bDumpEarlyConfigReads, bool bDumpEarlyPakFileReads)
+void DumpEarlyReads(bool bDumpEarlyConfigReads, bool bDumpEarlyPakFileReads, bool bForceQuitAfterEarlyReads)
 {
 	if (bDumpEarlyConfigReads)
 	{
@@ -3812,6 +3814,11 @@ void DumpEarlyReads(bool bDumpEarlyConfigReads, bool bDumpEarlyPakFileReads)
 		DumpRecordedFileReadsFromPaks();
 		DeleteRecordedFileReadsFromPaks();
 	}
+
+	if (bForceQuitAfterEarlyReads)
+	{
+		GEngine->DeferredCommands.Emplace(TEXT("Quit force"));
+	}
 }
 
 void HandleConfigReload(bool bReloadConfig)
