Added csv profiler scope for STAT_DeferredTickTime

#rnx

diff --git a/Engine/Source/Runtime/Launch/Private/LaunchEngineLoop.cpp b/Engine/Source/Runtime/Launch/Private/LaunchEngineLoop.cpp
--- a/Engine/Source/Runtime/Launch/Private/LaunchEngineLoop.cpp
+++ b/Engine/Source/Runtime/Launch/Private/LaunchEngineLoop.cpp
@@ -4477,6 +4477,7 @@ void FEngineLoop::Tick()
 		// tick core ticker, threads & deferred commands
 		{
 			SCOPE_CYCLE_COUNTER(STAT_DeferredTickTime);
+			CSV_SCOPED_TIMING_STAT_EXCLUSIVE(DeferredTickTime);
 			// Delete the objects which were enqueued for deferred cleanup before the previous frame.
 			delete PreviousPendingCleanupObjects;
 
