Prevent asking the user to build the editor if they're launching a blueprint project with an editor created from a shared build environment. We can infer the name of the editor executable in these cases, and fall back to the regular checks for whether modules are up-to-date.

#rnx

diff --git a/Engine/Source/Runtime/Launch/Private/LaunchEngineLoop.cpp b/Engine/Source/Runtime/Launch/Private/LaunchEngineLoop.cpp
--- a/Engine/Source/Runtime/Launch/Private/LaunchEngineLoop.cpp
+++ b/Engine/Source/Runtime/Launch/Private/LaunchEngineLoop.cpp
@@ -4664,23 +4664,47 @@ static void CheckForPrintTimesOverride()
 #if UE_EDITOR
 bool LaunchCorrectEditorExecutable(const FString& EditorTargetFileName)
 {
-	FTargetReceipt Receipt;
-	if(!FApp::IsUnattended() && FPaths::FileExists(EditorTargetFileName) && Receipt.Read(EditorTargetFileName))
+	// Don't allow relaunching the executable if we're running some unattended scripted process.
+	if(FApp::IsUnattended())
 	{
-		FString CurrentExecutableName = FPlatformProcess::ExecutablePath();
-		FPaths::MakeStandardFilename(CurrentExecutableName);
-
-		FString LaunchExecutableName = Receipt.Launch;
-		FPaths::MakeStandardFilename(LaunchExecutableName);
+		return false;
+	}
 
-		if(LaunchExecutableName != CurrentExecutableName)
+	// Figure out the executable that we should be running
+	FString LaunchExecutableName;
+	if(EditorTargetFileName.Len() == 0)
+	{
+		LaunchExecutableName = FPlatformProcess::GenerateApplicationPath(TEXT("UE4Editor"), FApp::GetBuildConfiguration());
+	}
+	else
+	{
+		FTargetReceipt Receipt;
+		if(!FPaths::FileExists(EditorTargetFileName) || !Receipt.Read(EditorTargetFileName))
 		{
-			UE_LOG(LogInit, Display, TEXT("Running incorrect executable for target. Launching %s..."), *LaunchExecutableName);
-			FPlatformProcess::CreateProc(*IFileManager::Get().ConvertToAbsolutePathForExternalAppForRead(*LaunchExecutableName), FCommandLine::GetOriginal(), true, false, false, nullptr, 0, nullptr, nullptr, nullptr);
-			return true;
+			return false;
 		}
+		LaunchExecutableName = Receipt.Launch;
 	}
-	return false;
+	FPaths::MakeStandardFilename(LaunchExecutableName);
+
+	// Get the current executable name. Don't allow relaunching if we're running the console app.
+	FString CurrentExecutableName = FPlatformProcess::ExecutablePath();
+	if(FPaths::GetBaseFilename(CurrentExecutableName).EndsWith(TEXT("-Cmd")))
+	{
+		return false;
+	}
+	FPaths::MakeStandardFilename(CurrentExecutableName);
+
+	// Nothing to do if they're the same
+	if(LaunchExecutableName == CurrentExecutableName)
+	{
+		return false;
+	}
+
+	// Relaunch the correct executable
+	UE_LOG(LogInit, Display, TEXT("Running incorrect executable for target. Launching %s..."), *LaunchExecutableName);
+	FPlatformProcess::CreateProc(*IFileManager::Get().ConvertToAbsolutePathForExternalAppForRead(*LaunchExecutableName), FCommandLine::GetOriginal(), true, false, false, nullptr, 0, nullptr, nullptr, nullptr);
+	return true;
 }
 #endif
 
@@ -4869,14 +4893,6 @@ bool FEngineLoop::AppInit( )
 			}
 		}
 
-		// If there was no editor target, use the default UE4Editor target
-		FString CompileProject = FPaths::GetProjectFilePath();
-		if(EditorTargetFileName.Len() == 0)
-		{
-			CompileProject.Empty();
-			EditorTargetFileName = FTargetReceipt::GetDefaultPath(FPlatformMisc::EngineDir(), TEXT("UE4Editor"), FPlatformProcess::GetBinariesSubdirectory(), FApp::GetBuildConfiguration(), nullptr);
-		}
-
 		// If we're not running the correct executable for the current target, and the listed executable exists, run that instead
 		if(LaunchCorrectEditorExecutable(EditorTargetFileName))
 		{
@@ -4948,6 +4964,13 @@ bool FEngineLoop::AppInit( )
 		
 		if(bNeedCompile)
 		{
+			// If there was no editor target, use the default UE4Editor target
+			FString CompileProject = FPaths::GetProjectFilePath();
+			if(EditorTargetFileName.Len() == 0)
+			{
+				CompileProject.Empty();
+			}
+
 			// Try to compile it
 			FFeedbackContext *Context = (FFeedbackContext*)FDesktopPlatformModule::Get()->GetNativeFeedbackContext();
 			Context->BeginSlowTask(FText::FromString(TEXT("Starting build...")), true, true);
