Missing check for relaunching the correct executable for blueprint projects.

#rnx

diff --git a/Engine/Source/Runtime/Launch/Private/LaunchEngineLoop.cpp b/Engine/Source/Runtime/Launch/Private/LaunchEngineLoop.cpp
--- a/Engine/Source/Runtime/Launch/Private/LaunchEngineLoop.cpp
+++ b/Engine/Source/Runtime/Launch/Private/LaunchEngineLoop.cpp
@@ -4950,7 +4950,7 @@ bool FEngineLoop::AppInit( )
 
 				bNeedCompile = true;
 			}
-			else if(!FPaths::FileExists(EditorTargetFileName))
+			else if(EditorTargetFileName.Len() > 0 && !FPaths::FileExists(EditorTargetFileName))
 			{
 				// Prompt to compile. The target file isn't essential, but we 
 				if (FPlatformMisc::MessageBoxExt(EAppMsgType::YesNo, *FString::Printf(TEXT("The %s file does not exist. Would you like to build the editor?"), *FPaths::GetCleanFilename(EditorTargetFileName)), TEXT("Missing target file")) == EAppReturnType::Yes)
