Allow installed editor builds to read the InstalledProject.txt file to determine which project they should load by default.

diff --git a/Engine/Source/Runtime/Launch/Private/LaunchEngineLoop.cpp b/Engine/Source/Runtime/Launch/Private/LaunchEngineLoop.cpp
--- a/Engine/Source/Runtime/Launch/Private/LaunchEngineLoop.cpp
+++ b/Engine/Source/Runtime/Launch/Private/LaunchEngineLoop.cpp
@@ -508,6 +508,26 @@ bool ParseGameProjectFromCommandLine(const TCHAR* InCmdLine, FString& OutProject
 	return false;
 }
 
+#if WITH_EDITOR
+bool ReadInstalledProjectPath(FString& OutProjFilePath)
+{
+	if(FApp::IsInstalled())
+	{
+		FString ProjFilePath;
+		if (FFileHelper::LoadFileToString(ProjFilePath, *(FPaths::RootDir() / TEXT("Engine/Build/InstalledProjectBuild.txt"))))
+		{
+			ProjFilePath.TrimStartAndEndInline();
+			if(ProjFilePath.Len() > 0)
+			{
+				OutProjFilePath = FPaths::RootDir() / ProjFilePath;
+				FPaths::NormalizeFilename(OutProjFilePath);
+				return true;
+			}
+		}
+	}
+	return false;
+}
+#endif
 
 bool LaunchSetGameName(const TCHAR *InCmdLine, FString& OutGameProjectFilePathUnnormalized)
 {
@@ -528,6 +548,18 @@ bool LaunchSetGameName(const TCHAR *InCmdLine, FString& OutGameProjectFilePathUn
 			OutGameProjectFilePathUnnormalized = ProjFilePath;
 			FPaths::SetProjectFilePath(ProjFilePath);
 		}
+#if WITH_EDITOR
+		else if(ReadInstalledProjectPath(ProjFilePath))
+		{
+			// Only set the game name if this is NOT a program...
+			if (FPlatformProperties::IsProgram() == false)
+			{
+				FApp::SetProjectName(*FPaths::GetBaseFilename(ProjFilePath));
+			}
+			OutGameProjectFilePathUnnormalized = ProjFilePath;
+			FPaths::SetProjectFilePath(ProjFilePath);
+		}
+#endif
 #if UE_GAME
 		else
 		{
