Prevent any attempts to rebuild the editor when run in unattended mode. Fixes issue where UBT is built on build machines before platform extensions have been registered.

diff --git a/Engine/Source/Runtime/Launch/Private/LaunchEngineLoop.cpp b/Engine/Source/Runtime/Launch/Private/LaunchEngineLoop.cpp
--- a/Engine/Source/Runtime/Launch/Private/LaunchEngineLoop.cpp
+++ b/Engine/Source/Runtime/Launch/Private/LaunchEngineLoop.cpp
@@ -4910,7 +4910,7 @@ bool FEngineLoop::AppInit( )
 
 	// Check whether the project or any of its plugins are missing or are out of date
 #if UE_EDITOR && !IS_MONOLITHIC
-	if(!GIsBuildMachine && FPaths::IsProjectFilePathSet())
+	if(!GIsBuildMachine && !FApp::IsUnattended() && FPaths::IsProjectFilePathSet())
 	{
 		// Check all the plugins are present
 		if(!PluginManager.AreRequiredPluginsAvailable())
