- Wake game thread when backgrounding event received

diff --git a/Engine/Source/Runtime/Launch/Private/Android/LaunchAndroid.cpp b/Engine/Source/Runtime/Launch/Private/Android/LaunchAndroid.cpp
--- a/Engine/Source/Runtime/Launch/Private/Android/LaunchAndroid.cpp
+++ b/Engine/Source/Runtime/Launch/Private/Android/LaunchAndroid.cpp
@@ -1193,6 +1193,8 @@ static void SuspendApp_EventThread()
 		EMDoneTrigger->Trigger();
 	}));
 
+	FEmbeddedCommunication::WakeGameThread();
+
 	FPreLoadScreenManager::EnableRendering(false);
 
 	// wait for a period of time before blocking rendering
