Make sure logging is flushed when using -ForceQuitAfterEarlyReads

diff --git a/Engine/Source/Runtime/Launch/Private/LaunchEngineLoop.cpp b/Engine/Source/Runtime/Launch/Private/LaunchEngineLoop.cpp
--- a/Engine/Source/Runtime/Launch/Private/LaunchEngineLoop.cpp
+++ b/Engine/Source/Runtime/Launch/Private/LaunchEngineLoop.cpp
@@ -3856,6 +3856,7 @@ void DumpEarlyReads(bool bDumpEarlyConfigReads, bool bDumpEarlyPakFileReads, boo
 
 	if (bForceQuitAfterEarlyReads)
 	{
+		GLog->Flush();
 		GEngine->DeferredCommands.Emplace(TEXT("Quit force"));
 	}
 }
