Don't try to load the CollisionAnalyzer module unless we're running an editor target. Fallout from allowing developer modules to be loaded in non-shipping configurations.

diff --git a/Engine/Source/Runtime/Launch/Private/LaunchEngineLoop.cpp b/Engine/Source/Runtime/Launch/Private/LaunchEngineLoop.cpp
--- a/Engine/Source/Runtime/Launch/Private/LaunchEngineLoop.cpp
+++ b/Engine/Source/Runtime/Launch/Private/LaunchEngineLoop.cpp
@@ -3308,12 +3308,14 @@ bool FEngineLoop::LoadStartupCoreModules()
 	{
 #if WITH_UNREAL_DEVELOPER_TOOLS
 		FModuleManager::Get().LoadModule("MessageLog");
+#endif	// WITH_UNREAL_DEVELOPER_TOOLS
+#if WITH_EDITOR
 		FModuleManager::Get().LoadModule("CollisionAnalyzer");
-#endif	//WITH_UNREAL_DEVELOPER_TOOLS
+#endif	// WITH_EDITOR
 	}
 
 #if WITH_UNREAL_DEVELOPER_TOOLS
-		FModuleManager::Get().LoadModule("FunctionalTesting");
+	FModuleManager::Get().LoadModule("FunctionalTesting");
 #endif	//WITH_UNREAL_DEVELOPER_TOOLS
 
 	SlowTask.EnterProgressFrame(30);
