Split PreInit() into PreInitPreStartupScreen() and PreInitPostStartupScreen() to add flexibility to engine preinit. This is useful on certain platforms which need to display a splash screen before preinit is completed.
The PostSplashScreen plugin can be used to show a custom splash screen during this time.

diff --git a/Engine/Plugins/Runtime/PostSplashScreen/Config/BasePostSplashScreen.ini b/Engine/Plugins/Runtime/PostSplashScreen/Config/BasePostSplashScreen.ini
new file mode 100644
--- /dev/null
+++ b/Engine/Plugins/Runtime/PostSplashScreen/Config/BasePostSplashScreen.ini
@@ -0,0 +1,35 @@
+[PreLoadScreen.UISettings]
+;How long will post splash screen be displayed in seconds.
+PostSplashScreenLength=0.01
+
+;--------------------------------------------------------------------------------------------------------------------------------------
+; Loading Groups
+; +LoadingGroups=LoadingGroupIdentifier -- List of LoadingGroups. If none are provided a default initial load is assumed
+;--------------------------------------------------------------------------------------------------------------------------------------
++LoadingGroups=PostSplashScreen
+
+
+;--------------------------------------------------------------------------------------------------------------------------------------
+; Custom Image Brushes
+; Format is +CustomImageBrushes=(Identifier,Filename,Width,Height, (LoadingGroupIdentifier1..., LoadingGroupidentifier2....)) or
+; Format is +CustomImageBrushes=(Identifier,Filename,Width,Height), and then LoadingGroupIdentifier is assumed to be the default
+; initial starting LoadingGroup
+;--------------------------------------------------------------------------------------------------------------------------------------
+
+;Background Brush Defines
++CustomImageBrushes=(SplashScreenBrush,Resources/SplashScreen.png,1920,1080,(PostSplashScreen))
+
+
+;--------------------------------------------------------------------------------------------------------------------------------------
+; ScreenGroupings
+; Used to describe a particular screen, which is a particular background with particular text at a given font size
+; Format is +ScreenGrouping(ScreenIdentifier, Brush Identifier, Text Identifier, Font Size)
+;--------------------------------------------------------------------------------------------------------------------------------------
++ScreenGroupings=(SG_SplashScreen,SplashScreenBrush,NullText,48)
+
+;--------------------------------------------------------------------------------------------------------------------------------------
+; ScreenOrders
+; Used to describe what order screens should be loaded in during a given LoadingGroup
+; Expected Format is : +ScreenOrders=(LoadingGroupIndentifier,ScreenGroupingsIdentifier)
+;--------------------------------------------------------------------------------------------------------------------------------------
++ScreenOrders=(PostSplashScreen,SG_SplashScreen)
diff --git a/Engine/Plugins/Runtime/PostSplashScreen/PostSplashScreen.uplugin b/Engine/Plugins/Runtime/PostSplashScreen/PostSplashScreen.uplugin
new file mode 100644
--- /dev/null
+++ b/Engine/Plugins/Runtime/PostSplashScreen/PostSplashScreen.uplugin
@@ -0,0 +1,27 @@
+{
+	"FileVersion": 3,
+	"Version": 1,
+	"VersionName": "1.0",
+	"FriendlyName": "Custom Splash Preload Screen",
+	"Description": "Handles displaying a custom splash screen after system splash screen during engine preinit.",
+	"Category": "PostSplashScreen",
+	"CreatedBy": "Epic Games, Inc.",
+	"CreatedByURL": "http://epicgames.com",
+	"DocsURL": "",
+	"MarketplaceURL": "",
+	"SupportURL": "",
+	"EnabledByDefault": true,
+	"CanContainContent": true,
+	"IsBetaVersion": false,
+	"Installed": false,
+	"Modules": [
+		{
+			"Name": "PostSplashScreen",
+			"Type": "RuntimeNoCommandlet",
+			"LoadingPhase": "PostSplashScreen",
+			"WhitelistPlatforms": [
+				"XboxOne"
+			]
+		}
+	]
+}
diff --git a/Engine/Plugins/Runtime/PostSplashScreen/Source/PostSplashScreen/PostSplashScreen.Build.cs b/Engine/Plugins/Runtime/PostSplashScreen/Source/PostSplashScreen/PostSplashScreen.Build.cs
new file mode 100644
--- /dev/null
+++ b/Engine/Plugins/Runtime/PostSplashScreen/Source/PostSplashScreen/PostSplashScreen.Build.cs
@@ -0,0 +1,28 @@
+// Copyright 1998-2019 Epic Games, Inc. All Rights Reserved.
+
+using UnrealBuildTool;
+
+// This module must be loaded "PreLoadingScreen" in the .uproject file, otherwise it will not hook in time!
+
+public class PostSplashScreen : ModuleRules
+{
+    public PostSplashScreen(ReadOnlyTargetRules Target) : base(Target)
+	{
+		PrivatePCHHeaderFile = "Public/PostSplashScreenPrivatePCH.h";
+
+        PrivateDependencyModuleNames.AddRange(
+			new string[] {
+                "Core",
+                "RHI",
+                "RenderCore",
+                "Slate",
+                "SlateCore",
+                "Projects",
+                "Engine",
+                "RenderCore",
+                "ApplicationCore",
+                "PreLoadScreen"
+            }
+		);
+    }
+}
diff --git a/Engine/Plugins/Runtime/PostSplashScreen/Source/PostSplashScreen/Private/CustomSplashScreen.cpp b/Engine/Plugins/Runtime/PostSplashScreen/Source/PostSplashScreen/Private/CustomSplashScreen.cpp
new file mode 100644
--- /dev/null
+++ b/Engine/Plugins/Runtime/PostSplashScreen/Source/PostSplashScreen/Private/CustomSplashScreen.cpp
@@ -0,0 +1,51 @@
+// Copyright 1998-2019 Epic Games, Inc. All Rights Reserved.
+
+#include "CustomSplashScreen.h"
+
+#include "PreLoadScreen.h"
+#include "PreLoadSettingsContainer.h"
+
+#include "Misc/ConfigCacheIni.h"
+
+
+#define LOCTEXT_NAMESPACE "CustomSplashScreen"
+
+#define CONFIG_BASENAME TEXT("PostSplashScreen")
+FString GPostSplashScreenIni;
+
+bool FCustomSplashScreen::Init()
+{
+	FPreLoadScreenBase::Init();
+	
+	SetPluginName("PostSplashScreen");
+	InitSettingsFromConfig("PostSplashScreen");
+
+	SAssignNew(SplashScreenWidget, SCustomSplashScreenWidget);
+
+	FPreLoadSettingsContainerBase::Get().LoadGrouping(TEXT("PostSplashScreen"));
+
+	TimeElapsed = 0.f;
+	MaxTimeToDisplay = 0.f;
+	FConfigCacheIni::LoadGlobalIniFile(GPostSplashScreenIni, CONFIG_BASENAME);
+	GConfig->GetFloat(TEXT("PreLoadScreen.UISettings"), TEXT("PostSplashScreenLength"), MaxTimeToDisplay, GPostSplashScreenIni);
+
+	return true;
+}
+
+void FCustomSplashScreen::Tick(float DeltaTime)
+{
+	FPreLoadScreenBase::Tick(DeltaTime);
+
+	TimeElapsed += DeltaTime;
+}
+
+bool FCustomSplashScreen::IsDone() const
+{
+	return TimeElapsed >= MaxTimeToDisplay;
+}
+
+float FCustomSplashScreen::GetAddedTickDelay()
+{
+	const float DefaultTickTime = 0.03f;
+	return FMath::Min(MaxTimeToDisplay, DefaultTickTime);
+}
\ No newline at end of file
diff --git a/Engine/Plugins/Runtime/PostSplashScreen/Source/PostSplashScreen/Private/CustomSplashScreenWidget.cpp b/Engine/Plugins/Runtime/PostSplashScreen/Source/PostSplashScreen/Private/CustomSplashScreenWidget.cpp
new file mode 100644
--- /dev/null
+++ b/Engine/Plugins/Runtime/PostSplashScreen/Source/PostSplashScreen/Private/CustomSplashScreenWidget.cpp
@@ -0,0 +1,56 @@
+// Copyright 1998-2019 Epic Games, Inc. All Rights Reserved.
+
+#include "CustomSplashScreenWidget.h"
+#include "PreLoadSettingsContainer.h"
+
+#include "Widgets/Layout/SBox.h"
+#include "Widgets/Layout/SScaleBox.h"
+#include "Widgets/Images/SImage.h"
+
+FCriticalSection SCustomSplashScreenWidget::BackgroundImageCrit;
+int SCustomSplashScreenWidget::CurrentBackgroundImage = 0;
+
+void SCustomSplashScreenWidget::Construct(const FArguments& InArgs)
+{
+    ChildSlot
+    [
+        SNew(SDPIScaler)
+        .DPIScale(this, &SCustomSplashScreenWidget::GetDPIScale)
+        [
+            SNew(SOverlay)
+
+            +SOverlay::Slot()
+            .HAlign(HAlign_Fill)
+            .VAlign(VAlign_Fill)
+            [
+                SNew(SVerticalBox)
+                    
+                + SVerticalBox::Slot()
+                [
+                    SNew(SOverlay)
+
+                    //Background Display
+                    +SOverlay::Slot()
+                    [
+                        SNew(SScaleBox)
+                        .Stretch(EStretch::Fill)
+                        [
+                            SNew(SImage)
+                            .Image(this, &SCustomSplashScreenWidget::GetCurrentBackgroundImage)
+                        ]
+                    ]
+                ]
+            ]
+        ]
+    ];
+}
+
+
+const FSlateBrush* SCustomSplashScreenWidget::GetCurrentBackgroundImage() const
+{
+    FScopeLock ScopeLock(&BackgroundImageCrit);
+
+    const FPreLoadSettingsContainerBase::FScreenGroupingBase* CurrentScreenIdentifier = FPreLoadSettingsContainerBase::Get().GetScreenAtIndex(CurrentBackgroundImage);
+    const FString& BackgroundBrushIdentifier = CurrentScreenIdentifier ? CurrentScreenIdentifier->ScreenBackgroundIdentifer : TEXT("");
+    return FPreLoadSettingsContainerBase::Get().GetBrush(BackgroundBrushIdentifier);
+}
diff --git a/Engine/Plugins/Runtime/PostSplashScreen/Source/PostSplashScreen/Private/PostSplashScreenModule.cpp b/Engine/Plugins/Runtime/PostSplashScreen/Source/PostSplashScreen/Private/PostSplashScreenModule.cpp
new file mode 100644
--- /dev/null
+++ b/Engine/Plugins/Runtime/PostSplashScreen/Source/PostSplashScreen/Private/PostSplashScreenModule.cpp
@@ -0,0 +1,57 @@
+// Copyright 1998-2019 Epic Games, Inc. All Rights Reserved.
+
+#include "PostSplashScreenPrivatePCH.h"
+
+#include "PreLoadScreenManager.h"
+
+#include "PostSplashScreen.h"
+#include "CustomSplashScreen.h"
+
+#include "RenderingThread.h"
+
+#define LOCTEXT_NAMESPACE "CustomSplashScreen"
+
+class FPostSplashScreenModule : public IPostSplashScreenModule
+{
+public:
+    FPostSplashScreenModule()
+	{
+	}
+
+	virtual void StartupModule() override
+	{
+		if (!GIsEditor && FApp::CanEverRender() && FPreLoadScreenManager::Get())
+		{
+			CustomSplashScreen = MakeShared<FCustomSplashScreen>();
+			if (CustomSplashScreen->Init())
+			{
+				FPreLoadScreenManager::Get()->RegisterPreLoadScreen(CustomSplashScreen);
+			}
+			else
+			{
+				CustomSplashScreen = nullptr;
+			}
+
+			FPreLoadScreenManager::Get()->OnPreLoadScreenManagerCleanUp.AddRaw(this, &FPostSplashScreenModule::CleanUpModule);
+		}
+    }
+
+    //Once the PreLoadScreenManager is cleaning up, we can get rid of all our resources too
+    virtual void CleanUpModule()
+    {
+        CustomSplashScreen.Reset();
+        ShutdownModule();
+    }
+
+	
+	virtual bool IsGameModule() const override
+	{
+		return true;
+	}
+
+    TSharedPtr<FCustomSplashScreen> CustomSplashScreen;
+};
+
+IMPLEMENT_GAME_MODULE(FPostSplashScreenModule, PostSplashScreen);
+
+#undef LOCTEXT_NAMESPACE
\ No newline at end of file
diff --git a/Engine/Plugins/Runtime/PostSplashScreen/Source/PostSplashScreen/Public/CustomSplashScreen.h b/Engine/Plugins/Runtime/PostSplashScreen/Source/PostSplashScreen/Public/CustomSplashScreen.h
new file mode 100644
--- /dev/null
+++ b/Engine/Plugins/Runtime/PostSplashScreen/Source/PostSplashScreen/Public/CustomSplashScreen.h
@@ -0,0 +1,39 @@
+// Copyright 1998-2019 Epic Games, Inc. All Rights Reserved.
+
+#pragma once
+
+#include "PostSplashScreenPrivatePCH.h"
+
+#include "Modules/ModuleInterface.h"
+#include "Modules/ModuleManager.h"
+#include "PreLoadScreenBase.h"
+
+#include "CustomSplashScreenWidget.h"
+
+
+class FCustomSplashScreen : public FPreLoadScreenBase
+{
+public:
+	
+    /*** IPreLoadScreen Implementation ***/
+	virtual void Tick(float DeltaTime) override;
+    virtual bool Init() override;
+
+    //Override to make sure this is always an EarlyStartupScreen
+    virtual EPreLoadScreenTypes GetPreLoadScreenType() const override { return EPreLoadScreenTypes::CustomSplashScreen; }
+    
+	virtual FName GetPreLoadScreenTag() const override { return NAME_None; }
+	
+	virtual bool IsDone() const override;
+
+	virtual float GetAddedTickDelay() override;
+
+    virtual TSharedPtr<SWidget> GetWidget() override { return SplashScreenWidget; }
+	
+private:
+	/** SharedPtr to our actual widget being rendered during ESP **/
+	TSharedPtr<SCustomSplashScreenWidget> SplashScreenWidget;
+
+	float TimeElapsed;
+	float MaxTimeToDisplay;
+};
diff --git a/Engine/Plugins/Runtime/PostSplashScreen/Source/PostSplashScreen/Public/CustomSplashScreenWidget.h b/Engine/Plugins/Runtime/PostSplashScreen/Source/PostSplashScreen/Public/CustomSplashScreenWidget.h
new file mode 100644
--- /dev/null
+++ b/Engine/Plugins/Runtime/PostSplashScreen/Source/PostSplashScreen/Public/CustomSplashScreenWidget.h
@@ -0,0 +1,37 @@
+// Copyright 1998-2019 Epic Games, Inc. All Rights Reserved.
+
+#pragma once
+
+#include "CoreMinimal.h"
+
+#include "Framework/Application/SlateApplication.h"
+#include "Rendering/SlateRenderer.h"
+
+#include "RenderingThread.h"
+
+#include "Widgets/DeclarativeSyntaxSupport.h"
+#include "Widgets/SWidget.h"
+#include "Widgets/SCompoundWidget.h"
+#include "Widgets/Layout/SDPIScaler.h"
+
+//Widget that displays a very simple version of a FPreLoadScreen UI that just includes a background and localized text together.
+//Rotates through the PreLoadScreens in the same order they are in the FPreLoadSettingsContainerBase. Uses the TimeToDisplayEachBackground variable to determine how long
+//to display each screen before rotating. Loops back when finished.
+class SCustomSplashScreenWidget : public SCompoundWidget
+{
+public:
+    SLATE_BEGIN_ARGS(SCustomSplashScreenWidget) {}
+    SLATE_END_ARGS()
+
+    void Construct(const FArguments& InArgs);
+
+private:
+    //Not used in the default simple implementation
+    virtual float GetDPIScale() const { return 1.0f; };
+
+    const FSlateBrush* GetCurrentBackgroundImage() const;
+	
+	static FCriticalSection BackgroundImageCrit;
+	static int CurrentBackgroundImage;
+
+};
\ No newline at end of file
diff --git a/Engine/Plugins/Runtime/PostSplashScreen/Source/PostSplashScreen/Public/PostSplashScreen.h b/Engine/Plugins/Runtime/PostSplashScreen/Source/PostSplashScreen/Public/PostSplashScreen.h
new file mode 100644
--- /dev/null
+++ b/Engine/Plugins/Runtime/PostSplashScreen/Source/PostSplashScreen/Public/PostSplashScreen.h
@@ -0,0 +1,33 @@
+// Copyright 1998-2019 Epic Games, Inc. All Rights Reserved.
+
+#pragma once
+
+#include "Modules/ModuleInterface.h"
+#include "Modules/ModuleManager.h"
+
+class IPostSplashScreenModule : public IModuleInterface
+{
+public:
+	/**
+	 * Singleton-like access to this module's interface.  This is just for convenience!
+	 * Beware of calling this during the shutdown phase, though.  Your module might have been unloaded already.
+	 *
+	 * @return Returns singleton instance, loading the module on demand if needed
+	 */
+	static inline IPostSplashScreenModule& Get()
+	{
+		static const FName ModuleName = TEXT("PostSplashScreen");
+		return FModuleManager::LoadModuleChecked< IPostSplashScreenModule >(ModuleName);
+	}
+
+	/**
+	 * Checks to see if this module is loaded and ready.  It is only valid to call Get() if IsAvailable() returns true.
+	 *
+	 * @return True if the module is loaded and ready to use
+	 */
+	static inline bool IsAvailable()
+	{
+		static const FName ModuleName = TEXT("PostSplashScreen");
+		return FModuleManager::Get().IsModuleLoaded(ModuleName);
+	}
+};
diff --git a/Engine/Plugins/Runtime/PostSplashScreen/Source/PostSplashScreen/Public/PostSplashScreenPrivatePCH.h b/Engine/Plugins/Runtime/PostSplashScreen/Source/PostSplashScreen/Public/PostSplashScreenPrivatePCH.h
new file mode 100644
--- /dev/null
+++ b/Engine/Plugins/Runtime/PostSplashScreen/Source/PostSplashScreen/Public/PostSplashScreenPrivatePCH.h
@@ -0,0 +1,9 @@
+// Copyright 1998-2019 Epic Games, Inc. All Rights Reserved.
+
+#pragma once
+
+#include "CoreMinimal.h"
+#include "PostSplashScreen.h"
+#include "Misc/FileHelper.h"
+
+DEFINE_LOG_CATEGORY_STATIC(LogPostSplashScreen, Display, All);
diff --git a/Engine/Source/Programs/UnrealBuildTool/System/ModuleDescriptor.cs b/Engine/Source/Programs/UnrealBuildTool/System/ModuleDescriptor.cs
--- a/Engine/Source/Programs/UnrealBuildTool/System/ModuleDescriptor.cs
+++ b/Engine/Source/Programs/UnrealBuildTool/System/ModuleDescriptor.cs
@@ -120,6 +120,11 @@ namespace UnrealBuildTool
 		/// </summary>
 		PostConfigInit,
 
+		/// <summary>
+		/// The first screen to be rendered after system splash screen
+		/// </summary>
+		PostSplashScreen,
+
 		/// <summary>
 		/// After PostConfigInit and before coreUobject initialized. used for early boot loading screens before the uobjects are initialized
 		/// </summary>
diff --git a/Engine/Source/Runtime/Core/Private/Misc/CoreDelegates.cpp b/Engine/Source/Runtime/Core/Private/Misc/CoreDelegates.cpp
--- a/Engine/Source/Runtime/Core/Private/Misc/CoreDelegates.cpp
+++ b/Engine/Source/Runtime/Core/Private/Misc/CoreDelegates.cpp
@@ -48,8 +48,6 @@ FCoreDelegates::FOnActorLabelChanged FCoreDelegates::OnActorLabelChanged;
 FCoreDelegates::FRegisterMovieStreamerDelegate FCoreDelegates::RegisterMovieStreamerDelegate;
 FCoreDelegates::FUnRegisterMovieStreamerDelegate FCoreDelegates::UnRegisterMovieStreamerDelegate;
 
-FCoreDelegates::FOnInitialLoadingScreenShown FCoreDelegates::OnInitialLoadingScreenShown;
-
 FCoreDelegates::FRegisterEncryptionKeyDelegate& FCoreDelegates::GetRegisterEncryptionKeyDelegate()
 {
 	static FRegisterEncryptionKeyDelegate RegisterEncryptionKeyDelegate;
diff --git a/Engine/Source/Runtime/Core/Public/Misc/CoreDelegates.h b/Engine/Source/Runtime/Core/Public/Misc/CoreDelegates.h
--- a/Engine/Source/Runtime/Core/Public/Misc/CoreDelegates.h
+++ b/Engine/Source/Runtime/Core/Public/Misc/CoreDelegates.h
@@ -131,10 +131,6 @@ public:
 	DECLARE_MULTICAST_DELEGATE_OneParam(FOnSyncLoadPackage, const FString&);
 	static FOnSyncLoadPackage OnSyncLoadPackage;
 
-	// Callback for when the initial loading screen is shown during PreInit
-	DECLARE_MULTICAST_DELEGATE(FOnInitialLoadingScreenShown);
-	static FOnInitialLoadingScreenShown OnInitialLoadingScreenShown;
-
 	// get a hotfix delegate
 	static FHotFixDelegate& GetHotfixDelegate(EHotfixDelegates::Type HotFix);
 
diff --git a/Engine/Source/Runtime/Launch/Private/LaunchEngineLoop.cpp b/Engine/Source/Runtime/Launch/Private/LaunchEngineLoop.cpp
--- a/Engine/Source/Runtime/Launch/Private/LaunchEngineLoop.cpp
+++ b/Engine/Source/Runtime/Launch/Private/LaunchEngineLoop.cpp
@@ -1197,9 +1197,10 @@ namespace AppLifetimeEventCapture
 #endif //WITH_ENGINE
 
 
-DECLARE_CYCLE_STAT( TEXT( "FEngineLoop::PreInit.AfterStats" ), STAT_FEngineLoop_PreInit_AfterStats, STATGROUP_LoadTime );
+DECLARE_CYCLE_STAT(TEXT("FEngineLoop::PreInitPreStartupScreen.AfterStats"), STAT_FEngineLoop_PreInitPreStartupScreen_AfterStats, STATGROUP_LoadTime);
+DECLARE_CYCLE_STAT(TEXT("FEngineLoop::PreInitPostStartupScreen.AfterStats"), STAT_FEngineLoop_PreInitPostStartupScreen_AfterStats, STATGROUP_LoadTime);
 
-int32 FEngineLoop::PreInit(const TCHAR* CmdLine)
+int32 FEngineLoop::PreInitPreStartupScreen(const TCHAR* CmdLine)
 {
 	TRACE_REGISTER_GAME_THREAD(FPlatformTLS::GetCurrentThreadId());
 #if CPUPROFILERTRACE_ENABLED
@@ -1755,7 +1756,7 @@ int32 FEngineLoop::PreInit(const TCHAR* CmdLine)
 	FThreadStats::StartThread();
 #endif
 
-	FScopeCycleCounter CycleCount_AfterStats(GET_STATID(STAT_FEngineLoop_PreInit_AfterStats));
+	FScopeCycleCounter CycleCount_AfterStats(GET_STATID(STAT_FEngineLoop_PreInitPreStartupScreen_AfterStats));
 
 	// Load Core modules required for everything else to work (needs to be loaded before InitializeRenderingCVarsCaching)
 	{
@@ -2242,7 +2243,7 @@ int32 FEngineLoop::PreInit(const TCHAR* CmdLine)
 
 	FEmbeddedCommunication::ForceTick(3);
 
-	FScopedSlowTask SlowTask(100, NSLOCTEXT("EngineLoop", "EngineLoop_Initializing", "Initializing..."));
+	FScopedSlowTask SlowTask(50, NSLOCTEXT("EngineLoop", "EngineLoop_Initializing", "PreInitPreStartupScreen..."));
 
 	SlowTask.EnterProgressFrame(10);
 
@@ -2334,10 +2335,10 @@ int32 FEngineLoop::PreInit(const TCHAR* CmdLine)
 		// if (!IsRunningCommandlet())
 		// hack: don't load global shaders if we are cooking we will load the shaders for the correct platform later
 		if (bEnableShaderCompile &&
-				!IsRunningDedicatedServer() &&
-				Commandline.Contains(TEXT("cookcommandlet")) == false &&
-				Commandline.Contains(TEXT("run=cook")) == false )
-		// if (FParse::Param(FCommandLine::Get(), TEXT("Multiprocess")) == false)
+			!IsRunningDedicatedServer() &&
+			Commandline.Contains(TEXT("cookcommandlet")) == false &&
+			Commandline.Contains(TEXT("run=cook")) == false)
+			// if (FParse::Param(FCommandLine::Get(), TEXT("Multiprocess")) == false)
 		{
 			LLM_SCOPE(ELLMTag::Shaders);
 			SCOPED_BOOT_TIMING("CompileGlobalShaderMap");
@@ -2358,12 +2359,12 @@ int32 FEngineLoop::PreInit(const TCHAR* CmdLine)
 			CreateMoviePlayer();
 		}
 
-        if (FPreLoadScreenManager::ArePreLoadScreensEnabled())
-        {
+		if (FPreLoadScreenManager::ArePreLoadScreensEnabled())
+		{
 			SCOPED_BOOT_TIMING("FPreLoadScreenManager::Create");
 			FPreLoadScreenManager::Create();
-            ensure(FPreLoadScreenManager::Get());
-        }
+			ensure(FPreLoadScreenManager::Get());
+		}
 
 		// If platforms support early movie playback we have to start the rendering thread much earlier
 #if PLATFORM_SUPPORTS_EARLY_MOVIE_PLAYBACK
@@ -2372,17 +2373,17 @@ int32 FEngineLoop::PreInit(const TCHAR* CmdLine)
 			PostInitRHI();
 		}
 
-		if(GUseThreadedRendering)
+		if (GUseThreadedRendering)
 		{
-			if(GRHISupportsRHIThread)
+			if (GRHISupportsRHIThread)
 			{
 				const bool DefaultUseRHIThread = true;
 				GUseRHIThread_InternalUseOnly = DefaultUseRHIThread;
-				if(FParse::Param(FCommandLine::Get(), TEXT("rhithread")))
+				if (FParse::Param(FCommandLine::Get(), TEXT("rhithread")))
 				{
 					GUseRHIThread_InternalUseOnly = true;
 				}
-				else if(FParse::Param(FCommandLine::Get(), TEXT("norhithread")))
+				else if (FParse::Param(FCommandLine::Get(), TEXT("norhithread")))
 				{
 					GUseRHIThread_InternalUseOnly = false;
 				}
@@ -2395,26 +2396,115 @@ int32 FEngineLoop::PreInit(const TCHAR* CmdLine)
 
 		FEmbeddedCommunication::ForceTick(4);
 
-#if !UE_SERVER// && !UE_EDITOR
-		if(!IsRunningDedicatedServer() && !IsRunningCommandlet())
 		{
-			TSharedRef<FSlateRenderer> SlateRenderer = GUsingNullRHI ?
-				FModuleManager::Get().LoadModuleChecked<ISlateNullRendererModule>("SlateNullRenderer").CreateSlateNullRenderer() :
-				FModuleManager::Get().GetModuleChecked<ISlateRHIRendererModule>("SlateRHIRenderer").CreateSlateRHIRenderer();
-
+#if !UE_SERVER// && !UE_EDITOR
+			if (!IsRunningDedicatedServer() && !IsRunningCommandlet())
 			{
-				SCOPED_BOOT_TIMING("CurrentSlateApp.InitializeRenderer");
-				// If Slate is being used, initialize the renderer after RHIInit
-				FSlateApplication& CurrentSlateApp = FSlateApplication::Get();
-				CurrentSlateApp.InitializeRenderer(SlateRenderer);
-			}
+				TSharedPtr<FSlateRenderer> SlateRenderer = GUsingNullRHI ?
+					FModuleManager::Get().LoadModuleChecked<ISlateNullRendererModule>("SlateNullRenderer").CreateSlateNullRenderer() :
+					FModuleManager::Get().GetModuleChecked<ISlateRHIRendererModule>("SlateRHIRenderer").CreateSlateRHIRenderer();
+				TSharedRef<FSlateRenderer> SlateRendererSharedRef = SlateRenderer.ToSharedRef();
 
-			{
-				SCOPED_BOOT_TIMING("FEngineFontServices::Create");
-				// Create the engine font services now that the Slate renderer is ready
-				FEngineFontServices::Create();
+				{
+					SCOPED_BOOT_TIMING("CurrentSlateApp.InitializeRenderer");
+					// If Slate is being used, initialize the renderer after RHIInit
+					FSlateApplication& CurrentSlateApp = FSlateApplication::Get();
+					CurrentSlateApp.InitializeRenderer(SlateRendererSharedRef);
+				}
+
+				{
+					SCOPED_BOOT_TIMING("FEngineFontServices::Create");
+					// Create the engine font services now that the Slate renderer is ready
+					FEngineFontServices::Create();
+				}
+
+				{
+					SCOPED_BOOT_TIMING("LoadModulesForProject(ELoadingPhase::PostSplashScreen)");
+					// Load up all modules that need to hook into the custom splash screen
+					if (!IProjectManager::Get().LoadModulesForProject(ELoadingPhase::PostSplashScreen) || !IPluginManager::Get().LoadModulesForEnabledPlugins(ELoadingPhase::PostSplashScreen))
+					{
+						return 1;
+					}
+				}
+
+				{
+					SCOPED_BOOT_TIMING("PlayFirstPreLoadScreen");
+
+					if (FPreLoadScreenManager::Get())
+					{
+						SCOPED_BOOT_TIMING("PlayFirstPreLoadScreen - FPreLoadScreenManager::Get()->Initialize");
+						// initialize and present custom splash screen
+						FPreLoadScreenManager::Get()->Initialize(SlateRendererSharedRef.Get());
+
+						if (FPreLoadScreenManager::Get()->HasRegisteredPreLoadScreenType(EPreLoadScreenTypes::CustomSplashScreen))
+						{
+							FPreLoadScreenManager::Get()->PlayFirstPreLoadScreen(EPreLoadScreenTypes::CustomSplashScreen);
+						}
+#if PLATFORM_XBOXONE && ENABLE_XBOXONE_FAST_ACTIVATION
+						else
+						{
+							UE_LOG(LogInit, Warning, TEXT("Enable fast activation without enabling a custom splash screen may cause garbage frame buffer being presented"));
+						}
+#endif
+					}
+				}
+
+				PreInitContext.SlateRenderer = SlateRenderer;
 			}
+#endif // !UE_SERVER
+		}
+	}
+#endif // WITH_ENGINE
+
+	// Save PreInitContext
+	PreInitContext.bDumpEarlyConfigReads = bDumpEarlyConfigReads;
+	PreInitContext.bDumpEarlyPakFileReads = bDumpEarlyPakFileReads;
+	PreInitContext.bForceQuitAfterEarlyReads = bForceQuitAfterEarlyReads;
+	PreInitContext.bWithConfigPatching = bWithConfigPatching;
+	PreInitContext.bHasEditorToken = bHasEditorToken;
+#if WITH_ENGINE
+	PreInitContext.bDisableDisregardForGC = bDisableDisregardForGC;
+	PreInitContext.bIsRegularClient = bIsRegularClient;
+#endif // WITH_ENGINE
+	PreInitContext.bTokenDoesNotHaveDash = bTokenDoesNotHaveDash;
+	PreInitContext.Token = Token;
+#if UE_EDITOR || WITH_ENGINE
+	PreInitContext.CommandletCommandLine = CommandletCommandLine;
+#endif // UE_EDITOR || WITH_ENGINE
+	PreInitContext.CommandLineCopy = CommandLineCopy;
+	
+	return 0;
+}
 
+int32 FEngineLoop::PreInitPostStartupScreen(const TCHAR* CmdLine)
+{
+	FScopedSlowTask SlowTask(50, NSLOCTEXT("EngineLoop", "EngineLoop_Initializing", "PreInitPostStartupScreen..."));
+	FScopeCycleCounter CycleCount_AfterStats(GET_STATID(STAT_FEngineLoop_PreInitPostStartupScreen_AfterStats));
+
+	// Restore PreInitContext
+	bool bDumpEarlyConfigReads = PreInitContext.bDumpEarlyConfigReads;
+	bool bDumpEarlyPakFileReads = PreInitContext.bDumpEarlyPakFileReads;
+	bool bForceQuitAfterEarlyReads = PreInitContext.bForceQuitAfterEarlyReads;
+	bool bWithConfigPatching = PreInitContext.bWithConfigPatching;
+	bool bHasEditorToken = PreInitContext.bHasEditorToken;
+#if WITH_ENGINE
+	bool bDisableDisregardForGC = PreInitContext.bDisableDisregardForGC;
+	bool bIsRegularClient = PreInitContext.bIsRegularClient;
+#endif // WITH_ENGINE
+	bool bTokenDoesNotHaveDash = PreInitContext.bTokenDoesNotHaveDash;
+	FString Token = PreInitContext.Token;
+#if UE_EDITOR || WITH_ENGINE
+	const TCHAR* CommandletCommandLine = PreInitContext.CommandletCommandLine;
+#endif // UE_EDITOR || WITH_ENGINE
+	TCHAR* CommandLineCopy = PreInitContext.CommandLineCopy;
+
+#if WITH_ENGINE
+	{
+#if !UE_SERVER// && !UE_EDITOR
+		if (!IsRunningDedicatedServer() && !IsRunningCommandlet())
+		{
+			TSharedPtr<FSlateRenderer> SlateRenderer = PreInitContext.SlateRenderer;
+			TSharedRef<FSlateRenderer> SlateRendererSharedRef = SlateRenderer.ToSharedRef();
 			{
 				SCOPED_BOOT_TIMING("GetMoviePlayer()->SetupLoadingScreenFromIni");
 				// allow the movie player to load a sequence from the .inis (a PreLoadingScreen module could have already initialized a sequence, in which case
@@ -2430,7 +2520,7 @@ int32 FEngineLoop::PreInit(const TCHAR* CmdLine)
 					return 1;
 				}
 			}
-			
+
 			IInstallBundleManager* BundleManager = IInstallBundleManager::GetPlatformInstallBundleManager();
 			if (BundleManager != nullptr && !BundleManager->IsNullInterface())
 			{
@@ -2442,16 +2532,16 @@ int32 FEngineLoop::PreInit(const TCHAR* CmdLine)
 			if (GetMoviePlayer()->HasEarlyStartupMovie())
 			{
 				SCOPED_BOOT_TIMING("EarlyStartupMovie");
-				GetMoviePlayer()->Initialize(SlateRenderer.Get());
+				GetMoviePlayer()->Initialize(SlateRendererSharedRef.Get());
 
-                // hide splash screen now before playing any movies
+				// hide splash screen now before playing any movies
 				FPlatformMisc::PlatformHandleSplashScreen(false);
 
 				// only allowed to play any movies marked as early startup.  These movies or widgets can have no interaction whatsoever with uobjects or engine features
 				GetMoviePlayer()->PlayEarlyStartupMovies();
 
-                // display the splash screen again now that early startup movies have played
-                FPlatformMisc::PlatformHandleSplashScreen(true);
+				// display the splash screen again now that early startup movies have played
+				FPlatformMisc::PlatformHandleSplashScreen(true);
 
 #if 0 && PAK_TRACKER// dump the files which have been accessed inside the pak file
 				FPakPlatformFile* PakPlatformFile = (FPakPlatformFile*)(FPlatformFileManager::Get().FindPlatformFile(FPakPlatformFile::GetTypeName()));
@@ -2508,51 +2598,46 @@ int32 FEngineLoop::PreInit(const TCHAR* CmdLine)
 				}
 #endif
 			}
-            else
-            {
+			else
+			{
 				SCOPED_BOOT_TIMING("PlayFirstPreLoadScreen");
 
-                if (FPreLoadScreenManager::Get())
-                {
+				if (FPreLoadScreenManager::Get())
+				{
 					SCOPED_BOOT_TIMING("PlayFirstPreLoadScreen - FPreLoadScreenManager::Get()->Initialize");
-                    // initialize and play our first Early PreLoad Screen if one is setup
-                    FPreLoadScreenManager::Get()->Initialize(SlateRenderer.Get());
+					// initialize and play our first Early PreLoad Screen if one is setup
+					FPreLoadScreenManager::Get()->Initialize(SlateRendererSharedRef.Get());
 
 					if (FPreLoadScreenManager::Get()->HasRegisteredPreLoadScreenType(EPreLoadScreenTypes::EarlyStartupScreen))
 					{
 						// disable the splash before playing the early startup screen
 						FPreLoadScreenManager::IsResponsibleForRenderingDelegate.AddLambda(
-							[](bool bIsPreloadScreenManResponsibleForRendering) 
-							{
-								FPlatformMisc::PlatformHandleSplashScreen(!bIsPreloadScreenManResponsibleForRendering);
-							}
+							[](bool bIsPreloadScreenManResponsibleForRendering)
+						{
+							FPlatformMisc::PlatformHandleSplashScreen(!bIsPreloadScreenManResponsibleForRendering);
+						}
 						);
-	                    FPreLoadScreenManager::Get()->PlayFirstPreLoadScreen(EPreLoadScreenTypes::EarlyStartupScreen);
-	                }
+						FPreLoadScreenManager::Get()->PlayFirstPreLoadScreen(EPreLoadScreenTypes::EarlyStartupScreen);
+					}
 					else
 					{
 						// no early startup screen, show the splash screen
 						FPlatformMisc::PlatformHandleSplashScreen(true);
 					}
-                }
+				}
 				else
 				{
 					// no preload manager, show the splash screen
 					FPlatformMisc::PlatformHandleSplashScreen(true);
 				}
-            }
-		
-			if (FCoreDelegates::OnInitialLoadingScreenShown.IsBound())
-			{
-				FCoreDelegates::OnInitialLoadingScreenShown.Broadcast();
 			}
 		}
-		else if ( IsRunningCommandlet() )
+		else if (IsRunningCommandlet())
 		{
 			// Create the engine font services now that the Slate renderer is ready
 			FEngineFontServices::Create();
 		}
-#endif
+#endif //!UE_SERVER
 
 		//Now that our EarlyStartupScreen is finished, lets take the necessary steps to mount paks, apply .ini cvars, and open the shader libraries if we installed content we expect to handle
 		//If using a bundle manager, assume its handling all this stuff and that we don't have to do it.
@@ -3176,10 +3261,27 @@ int32 FEngineLoop::PreInit(const TCHAR* CmdLine)
 
 	FEmbeddedCommunication::ForceTick(9);
 
+	PreInitContext.Cleanup();
+
 	// Note we still have 20% remaining on the slow task: this will be used by the Editor/Engine initialization next
 	return 0;
 }
 
+int32 FEngineLoop::PreInit(const TCHAR* CmdLine)
+{
+	if (int rv1 = PreInitPreStartupScreen(CmdLine) != 0)
+	{
+		return rv1;
+	}
+
+	if (int rv2 = PreInitPostStartupScreen(CmdLine) != 0)
+	{
+		return rv2;
+	}
+
+	return 0;
+}
+
 
 bool FEngineLoop::LoadCoreModules()
 {
diff --git a/Engine/Source/Runtime/Launch/Public/LaunchEngineLoop.h b/Engine/Source/Runtime/Launch/Public/LaunchEngineLoop.h
--- a/Engine/Source/Runtime/Launch/Public/LaunchEngineLoop.h
+++ b/Engine/Source/Runtime/Launch/Public/LaunchEngineLoop.h
@@ -11,8 +11,44 @@
 class FEngineService;
 class FPendingCleanupObjects;
 class ISessionService;
+class FSlateRenderer;
 
 
+struct FPreInitContext
+{
+	bool bDumpEarlyConfigReads;
+	bool bDumpEarlyPakFileReads;
+	bool bForceQuitAfterEarlyReads;
+	bool bWithConfigPatching;
+	bool bDisableDisregardForGC;
+	bool bHasEditorToken;
+	bool bIsRegularClient;
+	bool bTokenDoesNotHaveDash;
+
+	FString Token;
+	const TCHAR* CommandletCommandLine;
+	TCHAR* CommandLineCopy;
+
+	FPreInitContext() :
+		CommandletCommandLine(nullptr),
+		CommandLineCopy(nullptr)
+	{
+	}
+
+	void Cleanup()
+	{
+#if WITH_ENGINE && !UE_SERVER
+		SlateRenderer = nullptr;
+#endif // WITH_ENGINE && !UE_SERVER
+		CommandletCommandLine = nullptr;
+		CommandLineCopy = nullptr;
+	}
+
+#if WITH_ENGINE && !UE_SERVER
+	TSharedPtr<FSlateRenderer> SlateRenderer;
+#endif // WITH_ENGINE && !UE_SERVER
+};
+
 /**
  * Implements the main engine loop.	
  */
@@ -47,6 +83,12 @@ public:
 	 * @return The error level; 0 if successful, > 0 if there were errors.
 	 */ 
 	int32 PreInit(const TCHAR* CmdLine);
+	
+	/** First part of PreInit. */
+	int32 PreInitPreStartupScreen(const TCHAR* CmdLine);
+
+	/** Second part of PreInit. */
+	int32 PreInitPostStartupScreen(const TCHAR* CmdLine);
 
 	/** Load all modules needed before Init. */ 
 	void LoadPreInitModules();
@@ -151,6 +193,7 @@ private:
 	TSharedPtr<ISessionService> SessionService;
 
 #endif // WITH_ENGINE
+	FPreInitContext PreInitContext;
 };
 
 
diff --git a/Engine/Source/Runtime/PreLoadScreen/Private/PreLoadScreenManager.cpp b/Engine/Source/Runtime/PreLoadScreen/Private/PreLoadScreenManager.cpp
--- a/Engine/Source/Runtime/PreLoadScreen/Private/PreLoadScreenManager.cpp
+++ b/Engine/Source/Runtime/PreLoadScreen/Private/PreLoadScreenManager.cpp
@@ -112,6 +112,10 @@ void FPreLoadScreenManager::PlayPreLoadScreenAtIndex(int Index)
             {
                 HandleEngineLoadingPlay();
             }
+			else if (ActiveScreen->GetPreLoadScreenType() == EPreLoadScreenTypes::CustomSplashScreen)
+			{
+				HandleCustomSplashScreenPlay();
+			}
             else
             {
                 UE_LOG(LogPreLoadScreenManager, Fatal, TEXT("Attempting to play an Active PreLoadScreen type that hasn't been implemented inside of PreLoadScreenmanager!"));
@@ -204,6 +208,45 @@ void FPreLoadScreenManager::HandleEngineLoadingPlay()
     }
 }
 
+void FPreLoadScreenManager::HandleCustomSplashScreenPlay()
+{
+	if (ensureAlwaysMsgf(HasActivePreLoadScreenType(EPreLoadScreenTypes::CustomSplashScreen), TEXT("Invalid Active PreLoadScreen!")))
+	{
+		IPreLoadScreen* PreLoadScreen = GetActivePreLoadScreen();
+		if (PreLoadScreen && MainWindow.IsValid())
+		{
+			SCOPED_BOOT_TIMING("FPreLoadScreenManager::HandleCustomSplashScreenPlay()");
+
+			PreLoadScreen->OnPlay(MainWindow.Pin());
+
+			if (PreLoadScreen->GetWidget().IsValid())
+			{
+				MainWindow.Pin()->SetContent(PreLoadScreen->GetWidget().ToSharedRef());
+			}
+			
+			bool bDidDisableScreensaver = false;
+			if (FPlatformApplicationMisc::IsScreensaverEnabled())
+			{
+				bDidDisableScreensaver = FPlatformApplicationMisc::ControlScreensaver(FGenericPlatformApplicationMisc::EScreenSaverAction::Disable);
+			}
+
+			FPlatformMisc::HidePlatformStartupScreen();
+			FPlatformMisc::PlatformHandleSplashScreen(false);
+
+			while (!PreLoadScreen->IsDone())
+			{
+				EarlyPlayFrameTick();
+			}
+
+			if (bDidDisableScreensaver)
+			{
+				FPlatformApplicationMisc::ControlScreensaver(FGenericPlatformApplicationMisc::EScreenSaverAction::Enable);
+			}
+
+			StopPreLoadScreen();
+		}
+	}
+}
 
 void FPreLoadScreenManager::RenderTick()
 {
@@ -267,7 +310,7 @@ const IPreLoadScreen* FPreLoadScreenManager::GetActivePreLoadScreen() const
 
 void FPreLoadScreenManager::EarlyPlayFrameTick()
 {
-    if (ensureAlwaysMsgf(HasActivePreLoadScreenType(EPreLoadScreenTypes::EarlyStartupScreen), TEXT("EarlyPlayFrameTick called without a valid EarlyPreLoadScreen!")))
+    if (ensureAlwaysMsgf(HasActivePreLoadScreenType(EPreLoadScreenTypes::EarlyStartupScreen) || HasActivePreLoadScreenType(EPreLoadScreenTypes::CustomSplashScreen), TEXT("EarlyPlayFrameTick called without a valid EarlyPreLoadScreen!")))
     {
         GameLogicFrameTick();
         EarlyPlayRenderFrameTick();
diff --git a/Engine/Source/Runtime/PreLoadScreen/Public/PreLoadScreen.h b/Engine/Source/Runtime/PreLoadScreen/Public/PreLoadScreen.h
--- a/Engine/Source/Runtime/PreLoadScreen/Public/PreLoadScreen.h
+++ b/Engine/Source/Runtime/PreLoadScreen/Public/PreLoadScreen.h
@@ -6,6 +6,7 @@
 
 enum class EPreLoadScreenTypes : uint8
 {
+	CustomSplashScreen,
     EarlyStartupScreen,
     EngineLoadingScreen
 };
diff --git a/Engine/Source/Runtime/PreLoadScreen/Public/PreLoadScreenManager.h b/Engine/Source/Runtime/PreLoadScreen/Public/PreLoadScreenManager.h
--- a/Engine/Source/Runtime/PreLoadScreen/Public/PreLoadScreenManager.h
+++ b/Engine/Source/Runtime/PreLoadScreen/Public/PreLoadScreenManager.h
@@ -119,7 +119,10 @@ protected:
 	void PlatformSpecificGameLogicFrameTick();
 
     /*** These functions describe how everything is handled during an non-Early PreLoadPlay. Everything is handled asynchronously in this case with a standalone renderer ***/
-    void HandleEngineLoadingPlay();
+	void HandleEngineLoadingPlay();
+
+	/*** These functions describe the flow for showing an CustomSplashScreen ***/
+	void HandleCustomSplashScreenPlay();
 
     IPreLoadScreen* GetActivePreLoadScreen();
     const IPreLoadScreen* GetActivePreLoadScreen() const;
diff --git a/Engine/Source/Runtime/Projects/Private/ModuleDescriptor.cpp b/Engine/Source/Runtime/Projects/Private/ModuleDescriptor.cpp
--- a/Engine/Source/Runtime/Projects/Private/ModuleDescriptor.cpp
+++ b/Engine/Source/Runtime/Projects/Private/ModuleDescriptor.cpp
@@ -37,7 +37,10 @@ const TCHAR* ELoadingPhase::ToString( const ELoadingPhase::Type Value )
 
 	case PostConfigInit:
 		return TEXT( "PostConfigInit" );
-		
+
+	case PostSplashScreen:
+		return TEXT("PostSplashScreen");
+
 	case PreEarlyLoadingScreen:
 		return TEXT("PreEarlyLoadingScreen");
 
diff --git a/Engine/Source/Runtime/Projects/Public/ModuleDescriptor.h b/Engine/Source/Runtime/Projects/Public/ModuleDescriptor.h
--- a/Engine/Source/Runtime/Projects/Public/ModuleDescriptor.h
+++ b/Engine/Source/Runtime/Projects/Public/ModuleDescriptor.h
@@ -21,6 +21,9 @@ namespace ELoadingPhase
 		/** Loaded before the engine is fully initialized, immediately after the config system has been initialized.  Necessary only for very low-level hooks */
 		PostConfigInit,
 
+		/** The first screen to be rendered after system splash screen */
+		PostSplashScreen,
+
 		/** Loaded before coreUObject for setting up manual loading screens, used for our chunk patching system */
 		PreEarlyLoadingScreen,
 
