Fix static analysis warning V593: Consider reviewing the expression of the 'A = B != C' kind. The expression is calculated as following: 'A = (B != C)'.
#rnx

diff --git a/Engine/Source/Runtime/Launch/Private/LaunchEngineLoop.cpp b/Engine/Source/Runtime/Launch/Private/LaunchEngineLoop.cpp
--- a/Engine/Source/Runtime/Launch/Private/LaunchEngineLoop.cpp
+++ b/Engine/Source/Runtime/Launch/Private/LaunchEngineLoop.cpp
@@ -3311,12 +3311,14 @@ int32 FEngineLoop::PreInitPostStartupScreen(const TCHAR* CmdLine)
 
 int32 FEngineLoop::PreInit(const TCHAR* CmdLine)
 {
-	if (int rv1 = PreInitPreStartupScreen(CmdLine) != 0)
+	const int32 rv1 = PreInitPreStartupScreen(CmdLine);
+	if (rv1 != 0)
 	{
 		return rv1;
 	}
 
-	if (int rv2 = PreInitPostStartupScreen(CmdLine) != 0)
+	const int32 rv2 = PreInitPostStartupScreen(CmdLine);
+	if (rv2 != 0)
 	{
 		return rv2;
 	}
