#DEV-VR - Moving FunctionalTesting dependency under bCompileAgainstEngine

diff --git a/Engine/Source/Runtime/Launch/Launch.Build.cs b/Engine/Source/Runtime/Launch/Launch.Build.cs
--- a/Engine/Source/Runtime/Launch/Launch.Build.cs
+++ b/Engine/Source/Runtime/Launch/Launch.Build.cs
@@ -102,8 +102,7 @@ public class Launch : ModuleRules
 					"NetworkFile",
 					"StreamingFile",
 					"CookedIterativeFile",
-					"AutomationWorker",
-					"FunctionalTesting"
+					"AutomationWorker"
 			});
 		}
 
@@ -141,7 +140,8 @@ public class Launch : ModuleRules
 			DynamicallyLoadedModuleNames.AddRange(new string[] {
 				"TaskGraph",
 				"RealtimeProfiler",
-				"ProfilerService"
+				"ProfilerService",
+				"FunctionalTesting"
 			});
 		}
 
