#DEV-VR - Moving FunctionalTest dependency to be both non-shipping and CompileAgainstEngine.  Allowing FunctionalTesting to be compiled in binary build

diff --git a/Engine/Source/Developer/FunctionalTesting/FunctionalTesting.Build.cs b/Engine/Source/Developer/FunctionalTesting/FunctionalTesting.Build.cs
--- a/Engine/Source/Developer/FunctionalTesting/FunctionalTesting.Build.cs
+++ b/Engine/Source/Developer/FunctionalTesting/FunctionalTesting.Build.cs
@@ -44,5 +44,11 @@ public class FunctionalTesting : ModuleRules
                 }
             );
         }
-    }
+
+		//make sure this is compiled for binary builds
+		if (Target.Configuration != UnrealTargetConfiguration.Shipping)
+		{
+			PrecompileForTargets = PrecompileTargetsType.Any;
+		}
+	}
 }
diff --git a/Engine/Source/Runtime/Launch/Launch.Build.cs b/Engine/Source/Runtime/Launch/Launch.Build.cs
--- a/Engine/Source/Runtime/Launch/Launch.Build.cs
+++ b/Engine/Source/Runtime/Launch/Launch.Build.cs
@@ -131,6 +131,13 @@ public class Launch : ModuleRules
 				"ClothingSystemRuntime",
 				"ClothingSystemRuntimeInterface"
 			});
+
+			if (Target.Configuration != UnrealTargetConfiguration.Shipping)
+			{
+				PrivateDependencyModuleNames.AddRange(new string[] {
+					"FunctionalTesting"
+				});
+			}
 		}
 
 		if (Target.Configuration != UnrealTargetConfiguration.Shipping)
@@ -140,8 +147,7 @@ public class Launch : ModuleRules
 			DynamicallyLoadedModuleNames.AddRange(new string[] {
 				"TaskGraph",
 				"RealtimeProfiler",
-				"ProfilerService",
-				"FunctionalTesting"
+				"ProfilerService"
 			});
 		}
 
