Copying //UE4/Dev-Main to Dev-Rendering (//UE4/Dev-Rendering)

diff --git a/Engine/Source/Runtime/Launch/Private/LaunchEngineLoop.cpp b/Engine/Source/Runtime/Launch/Private/LaunchEngineLoop.cpp
--- a/Engine/Source/Runtime/Launch/Private/LaunchEngineLoop.cpp
+++ b/Engine/Source/Runtime/Launch/Private/LaunchEngineLoop.cpp
@@ -1817,8 +1817,12 @@ int32 FEngineLoop::PreInit(const TCHAR* CmdLine)
 #if WITH_EDITOR
 		{
 			TRACE_THREAD_GROUP_SCOPE("LargeThreadPool");
-			// GLargeThreadPool needs extra stack size just like regular thread pool in editor mode, Crunch texture compression relies on this 
-			verify(GLargeThreadPool->Create(NumThreadsInLargeThreadPool, StackSize * 1024));
+			// when we are in the editor we like to do things like build lighting and such
+			// this thread pool can be used for those purposes
+			GLargeThreadPool = FQueuedThreadPool::Allocate();
+			int32 NumThreadsInLargeThreadPool = FMath::Max(FPlatformMisc::NumberOfCoresIncludingHyperthreads() - 2, 2);
+
+			verify(GLargeThreadPool->Create(NumThreadsInLargeThreadPool, 128 * 1024));
 		}
 #endif
 	}
