Fix a crash at startup in Windows build.
This issue lead to a crash when game was launched directly from game executable when a launcher is available.  Now PreInit() will return 1 instead of 0 when this happens.
Additional guard is added to PreInitPostStartupScreen in case there are more cases in future where PreInitPreStartupScreen returns 0 while requesting exit.

diff --git a/Engine/Source/Runtime/Launch/Private/LaunchEngineLoop.cpp b/Engine/Source/Runtime/Launch/Private/LaunchEngineLoop.cpp
--- a/Engine/Source/Runtime/Launch/Private/LaunchEngineLoop.cpp
+++ b/Engine/Source/Runtime/Launch/Private/LaunchEngineLoop.cpp
@@ -1299,7 +1299,7 @@ int32 FEngineLoop::PreInitPreStartupScreen(const TCHAR* CmdLine)
 		ILauncherCheckModule::Get().RunLauncher(ELauncherAction::AppLaunch);
 		// We wish to exit
 		RequestEngineExit(TEXT("Run outside of launcher; restarting via launcher"));
-		return 0;
+		return 1;
 	}
 #endif
 
@@ -2481,7 +2481,12 @@ int32 FEngineLoop::PreInitPreStartupScreen(const TCHAR* CmdLine)
 
 int32 FEngineLoop::PreInitPostStartupScreen(const TCHAR* CmdLine)
 {
-	FScopedSlowTask SlowTask(50, NSLOCTEXT("EngineLoop", "EngineLoop_Initializing_PreInitPostStartupScreen", "PreInitPostStartupScreen..."));
+	if (IsEngineExitRequested())
+	{
+		return 0;
+	}
+
+	FScopedSlowTask SlowTask(50, NSLOCTEXT("EngineLoop", "EngineLoop_Initializing", "PreInitPostStartupScreen..."));
 	FScopeCycleCounter CycleCount_AfterStats(GET_STATID(STAT_FEngineLoop_PreInitPostStartupScreen_AfterStats));
 
 	// Restore PreInitContext
@@ -3429,7 +3434,7 @@ bool FEngineLoop::LoadStartupCoreModules()
 	}
 
 #if WITH_UNREAL_DEVELOPER_TOOLS
-	FModuleManager::Get().LoadModule("FunctionalTesting");
+		FModuleManager::Get().LoadModule("FunctionalTesting");
 #endif	//WITH_UNREAL_DEVELOPER_TOOLS
 
 	SlowTask.EnterProgressFrame(30);
@@ -3812,7 +3817,7 @@ void FEngineLoop::Exit()
 		}
 		if (bFlushOnExit)
 		{
-			FlushAsyncLoading();
+	FlushAsyncLoading();
 		}
 		else
 		{
@@ -5074,7 +5079,7 @@ bool FEngineLoop::AppInit( )
 		// Find the editor target
 		FString EditorTargetFileName;
 		for (const FTargetInfo& Target : FDesktopPlatformModule::Get()->GetTargetsForProject(FPaths::GetProjectFilePath()))
-		{
+	{
 			if (Target.Type == EBuildTargetType::Editor)
 			{
 				// Read the editor target receipt
