Fix for race condition that will likely cause a crash if the audio mixer is hotfixed on or off on iOS

diff --git a/Engine/Source/Runtime/Launch/Private/IOS/LaunchIOS.cpp b/Engine/Source/Runtime/Launch/Private/IOS/LaunchIOS.cpp
--- a/Engine/Source/Runtime/Launch/Private/IOS/LaunchIOS.cpp
+++ b/Engine/Source/Runtime/Launch/Private/IOS/LaunchIOS.cpp
@@ -67,11 +67,14 @@ void FAppEntry::Suspend(bool bIsInterrupt)
 			{
 				if (FTaskGraphInterface::IsRunning() && !IsEngineExitRequested())
 				{
-					FFunctionGraphTask::CreateAndDispatchWhenReady([AudioDevice]()
+					FFunctionGraphTask::CreateAndDispatchWhenReady([]()
 					{
-						FAudioThread::RunCommandOnAudioThread([AudioDevice]()
+						FAudioThread::RunCommandOnAudioThread([]()
 						{
-							AudioDevice->SetTransientMasterVolume(0.0f);
+							if (GEngine && GEngine->GetMainAudioDevice())
+							{
+								GEngine->GetMainAudioDevice()->SetTransientMasterVolume(0.0f);
+							}
 						}, TStatId());
 					}, TStatId(), NULL, ENamedThreads::GameThread);
 				}
@@ -96,11 +99,14 @@ void FAppEntry::Suspend(bool bIsInterrupt)
 
 				if (FTaskGraphInterface::IsRunning())
 				{
-					FGraphEventRef ResignTask = FFunctionGraphTask::CreateAndDispatchWhenReady([AudioDevice]()
+					FGraphEventRef ResignTask = FFunctionGraphTask::CreateAndDispatchWhenReady([]()
 					{
-						FAudioThread::RunCommandOnAudioThread([AudioDevice]()
+						FAudioThread::RunCommandOnAudioThread([]()
 						{
-							AudioDevice->SuspendContext();
+							if (GEngine && GEngine->GetMainAudioDevice())
+							{
+								GEngine->GetMainAudioDevice()->SuspendContext();
+							}
 						}, TStatId());
                 
 						FAudioCommandFence AudioCommandFence;
@@ -156,11 +162,14 @@ void FAppEntry::Resume(bool bIsInterrupt)
 			{
 				if (FTaskGraphInterface::IsRunning())
 				{
-					FFunctionGraphTask::CreateAndDispatchWhenReady([AudioDevice]()
+					FFunctionGraphTask::CreateAndDispatchWhenReady([]()
 					{
-						FAudioThread::RunCommandOnAudioThread([AudioDevice]()
+						FAudioThread::RunCommandOnAudioThread([]()
 						{
-							AudioDevice->SetTransientMasterVolume(1.0f);
+							if (GEngine && GEngine->GetMainAudioDevice())
+							{
+								GEngine->GetMainAudioDevice()->SetTransientMasterVolume(1.0f);
+							}
 						}, TStatId());
 					}, TStatId(), NULL, ENamedThreads::GameThread);
 				}
@@ -201,11 +210,14 @@ void FAppEntry::ResumeAudioContext()
 		{
 			if (FTaskGraphInterface::IsRunning())
 			{
-				FFunctionGraphTask::CreateAndDispatchWhenReady([AudioDevice]()
+				FFunctionGraphTask::CreateAndDispatchWhenReady([]()
 				{
-					FAudioThread::RunCommandOnAudioThread([AudioDevice]()
+					FAudioThread::RunCommandOnAudioThread([]()
 					{
-						AudioDevice->ResumeContext();
+						if (GEngine && GEngine->GetMainAudioDevice())
+						{
+							GEngine->GetMainAudioDevice()->ResumeContext();
+						}
 					}, TStatId());
 				}, TStatId(), NULL, ENamedThreads::GameThread);
 			}
@@ -228,11 +240,14 @@ void FAppEntry::RestartAudio()
             //increment the counter, otherwise ResumeContext won't work
             IncrementAudioSuspendCounters();
 
-			FFunctionGraphTask::CreateAndDispatchWhenReady([AudioDevice]()
+			FFunctionGraphTask::CreateAndDispatchWhenReady([]()
 			{
-				FAudioThread::RunCommandOnAudioThread([AudioDevice]()
+				FAudioThread::RunCommandOnAudioThread([]()
 				{
-					AudioDevice->ResumeContext();
+					if (GEngine && GEngine->GetMainAudioDevice())
+					{
+						GEngine->GetMainAudioDevice()->ResumeContext();
+					}
 				}, TStatId());
 			}, TStatId(), NULL, ENamedThreads::GameThread);
 		}
