Fix engine pre-init to only use a single slow task with a nice display name

#rnx

diff --git a/Engine/Source/Runtime/Launch/Private/LaunchEngineLoop.cpp b/Engine/Source/Runtime/Launch/Private/LaunchEngineLoop.cpp
--- a/Engine/Source/Runtime/Launch/Private/LaunchEngineLoop.cpp
+++ b/Engine/Source/Runtime/Launch/Private/LaunchEngineLoop.cpp
@@ -2246,7 +2246,8 @@ int32 FEngineLoop::PreInitPreStartupScreen(const TCHAR* CmdLine)
 
 	FEmbeddedCommunication::ForceTick(3);
 
-	FScopedSlowTask SlowTask(50, NSLOCTEXT("EngineLoop", "EngineLoop_Initializing_PreInitPreStartupScreen", "PreInitPreStartupScreen..."));
+	PreInitContext.SlowTaskPtr = new FScopedSlowTask(100, NSLOCTEXT("EngineLoop", "EngineLoop_Initializing", "Initializing..."));
+	FScopedSlowTask& SlowTask = *PreInitContext.SlowTaskPtr;
 
 	SlowTask.EnterProgressFrame(10);
 
@@ -2486,7 +2487,6 @@ int32 FEngineLoop::PreInitPostStartupScreen(const TCHAR* CmdLine)
 		return 0;
 	}
 
-	FScopedSlowTask SlowTask(50, NSLOCTEXT("EngineLoop", "EngineLoop_Initializing", "PreInitPostStartupScreen..."));
 	FScopeCycleCounter CycleCount_AfterStats(GET_STATID(STAT_FEngineLoop_PreInitPostStartupScreen_AfterStats));
 
 	// Restore PreInitContext
@@ -2505,6 +2505,7 @@ int32 FEngineLoop::PreInitPostStartupScreen(const TCHAR* CmdLine)
 	const TCHAR* CommandletCommandLine = PreInitContext.CommandletCommandLine;
 #endif // UE_EDITOR || WITH_ENGINE
 	TCHAR* CommandLineCopy = PreInitContext.CommandLineCopy;
+	FScopedSlowTask& SlowTask = *PreInitContext.SlowTaskPtr;
 
 #if WITH_ENGINE
 	{
@@ -3223,8 +3224,6 @@ int32 FEngineLoop::PreInitPostStartupScreen(const TCHAR* CmdLine)
 	}
 #endif // !UE_BUILD_SHIPPING
 
-	delete [] CommandLineCopy;
-
 	// initialize the pointer, as it is deleted before being assigned in the first frame
 	PendingCleanupObjects = nullptr;
 
@@ -3280,12 +3279,14 @@ int32 FEngineLoop::PreInit(const TCHAR* CmdLine)
 	const int32 rv1 = PreInitPreStartupScreen(CmdLine);
 	if (rv1 != 0)
 	{
+		PreInitContext.Cleanup();
 		return rv1;
 	}
 
 	const int32 rv2 = PreInitPostStartupScreen(CmdLine);
 	if (rv2 != 0)
 	{
+		PreInitContext.Cleanup();
 		return rv2;
 	}
 
@@ -5494,4 +5495,18 @@ void FEngineLoop::PreInitHMDDevice()
 #endif // #if WITH_ENGINE && !UE_SERVER
 }
 
+void FPreInitContext::Cleanup()
+{
+#if WITH_ENGINE && !UE_SERVER
+	SlateRenderer = nullptr;
+#endif // WITH_ENGINE && !UE_SERVER
+	CommandletCommandLine = nullptr;
+
+	delete[] CommandLineCopy;
+	CommandLineCopy = nullptr;
+
+	delete SlowTaskPtr;
+	SlowTaskPtr = nullptr;
+}
+
 #undef LOCTEXT_NAMESPACE
diff --git a/Engine/Source/Runtime/Launch/Public/LaunchEngineLoop.h b/Engine/Source/Runtime/Launch/Public/LaunchEngineLoop.h
--- a/Engine/Source/Runtime/Launch/Public/LaunchEngineLoop.h
+++ b/Engine/Source/Runtime/Launch/Public/LaunchEngineLoop.h
@@ -13,36 +13,26 @@ class FPendingCleanupObjects;
 class ISessionService;
 class FSlateRenderer;
 
+struct FScopedSlowTask;
 
 struct FPreInitContext
 {
-	bool bDumpEarlyConfigReads;
-	bool bDumpEarlyPakFileReads;
-	bool bForceQuitAfterEarlyReads;
-	bool bWithConfigPatching;
-	bool bDisableDisregardForGC;
-	bool bHasEditorToken;
-	bool bIsRegularClient;
-	bool bTokenDoesNotHaveDash;
+	bool bDumpEarlyConfigReads = false;
+	bool bDumpEarlyPakFileReads = false;
+	bool bForceQuitAfterEarlyReads = false;
+	bool bWithConfigPatching = false;
+	bool bDisableDisregardForGC = false;
+	bool bHasEditorToken = false;
+	bool bIsRegularClient = false;
+	bool bTokenDoesNotHaveDash = false;
 
 	FString Token;
-	const TCHAR* CommandletCommandLine;
-	TCHAR* CommandLineCopy;
+	const TCHAR* CommandletCommandLine = nullptr;
+	TCHAR* CommandLineCopy = nullptr;
 
-	FPreInitContext() :
-		CommandletCommandLine(nullptr),
-		CommandLineCopy(nullptr)
-	{
-	}
+	FScopedSlowTask* SlowTaskPtr = nullptr;
 
-	void Cleanup()
-	{
-#if WITH_ENGINE && !UE_SERVER
-		SlateRenderer = nullptr;
-#endif // WITH_ENGINE && !UE_SERVER
-		CommandletCommandLine = nullptr;
-		CommandLineCopy = nullptr;
-	}
+	void Cleanup();
 
 #if WITH_ENGINE && !UE_SERVER
 	TSharedPtr<FSlateRenderer> SlateRenderer;
