Clean up PreInit context when PreInitPreStartupScreen or PreInitPostStartupScreen fails.

diff --git a/Engine/Source/Runtime/Launch/Private/LaunchEngineLoop.cpp b/Engine/Source/Runtime/Launch/Private/LaunchEngineLoop.cpp
--- a/Engine/Source/Runtime/Launch/Private/LaunchEngineLoop.cpp
+++ b/Engine/Source/Runtime/Launch/Private/LaunchEngineLoop.cpp
@@ -3308,6 +3308,11 @@ bool FEngineLoop::LoadCoreModules()
 }
 
 
+void FEngineLoop::CleanupPreInitContext()
+{
+	PreInitContext.Cleanup();
+}
+
 void FEngineLoop::LoadPreInitModules()
 {
 	DECLARE_SCOPE_CYCLE_COUNTER(TEXT("Loading PreInit Modules"), STAT_PreInitModules, STATGROUP_LoadTime);
diff --git a/Engine/Source/Runtime/Launch/Public/LaunchEngineLoop.h b/Engine/Source/Runtime/Launch/Public/LaunchEngineLoop.h
--- a/Engine/Source/Runtime/Launch/Public/LaunchEngineLoop.h
+++ b/Engine/Source/Runtime/Launch/Public/LaunchEngineLoop.h
@@ -86,6 +86,9 @@ public:
 	/** Load core modules. */
 	bool LoadCoreModules();
 
+	/** Clean up PreInit context. */
+	void CleanupPreInitContext();
+
 #if WITH_ENGINE
 	
 	/** Load all core modules needed at startup time. */
