Fix prompt to rebuild the editor when running a content-only project. Computed path to the editor's target file was incorrect.

#rnx

diff --git a/Engine/Source/Runtime/Core/Private/Misc/Paths.cpp b/Engine/Source/Runtime/Core/Private/Misc/Paths.cpp
--- a/Engine/Source/Runtime/Core/Private/Misc/Paths.cpp
+++ b/Engine/Source/Runtime/Core/Private/Misc/Paths.cpp
@@ -1442,6 +1442,29 @@ bool FPaths::IsSamePath(const FString& PathA, const FString& PathB)
 #endif
 }
 
+bool FPaths::IsUnderDirectory(const FString& InPath, const FString& InDirectory)
+{
+	FString Path = InPath;
+	MakeStandardFilename(Path);
+
+	FString Directory = InDirectory;
+	MakeStandardFilename(Directory);
+
+	if (Directory.EndsWith(TEXT("/")))
+	{
+		Directory.RemoveAt(Directory.Len() - 1);
+	}
+
+#if PLATFORM_WINDOWS || PLATFORM_XBOXONE || PLATFORM_HOLOLENS
+	int Compare = FCString::Strnicmp(*Path, *Directory, Directory.Len());
+#else
+	int Compare = FCString::Strncmp(*Path, *Directory, Directory.Len());
+#endif
+
+	return Compare == 0 && (Path[Directory.Len()] == 0 || Path[Directory.Len()] == '/');
+}
+
+
 void FPaths::TearDown()
 {
 	TLazySingleton<FStaticData>::TearDown();
diff --git a/Engine/Source/Runtime/Core/Public/Misc/Paths.h b/Engine/Source/Runtime/Core/Public/Misc/Paths.h
--- a/Engine/Source/Runtime/Core/Public/Misc/Paths.h
+++ b/Engine/Source/Runtime/Core/Public/Misc/Paths.h
@@ -465,6 +465,9 @@ public:
 	 */
 	static bool IsSamePath(const FString& PathA, const FString& PathB);
 
+	/** Determines if a path is under a given directory */
+	static bool IsUnderDirectory(const FString& InPath, const FString& InDirectory);
+
 	/** Normalize all / and \ to TEXT("/") and remove any trailing TEXT("/") if the character before that is not a TEXT("/") or a colon */
 	static void NormalizeDirectoryName(FString& InPath);
 
diff --git a/Engine/Source/Runtime/Launch/Private/LaunchEngineLoop.cpp b/Engine/Source/Runtime/Launch/Private/LaunchEngineLoop.cpp
--- a/Engine/Source/Runtime/Launch/Private/LaunchEngineLoop.cpp
+++ b/Engine/Source/Runtime/Launch/Private/LaunchEngineLoop.cpp
@@ -5092,8 +5092,14 @@ bool FEngineLoop::AppInit( )
 	{
 			if (Target.Type == EBuildTargetType::Editor)
 			{
-				// Read the editor target receipt
-				EditorTargetFileName = FTargetReceipt::GetDefaultPath(FPlatformMisc::ProjectDir(), *Target.Name, FPlatformProcess::GetBinariesSubdirectory(), FApp::GetBuildConfiguration(), nullptr);
+				if(FPaths::IsUnderDirectory(Target.Path, FPlatformMisc::ProjectDir()))
+				{
+					EditorTargetFileName = FTargetReceipt::GetDefaultPath(FPlatformMisc::ProjectDir(), *Target.Name, FPlatformProcess::GetBinariesSubdirectory(), FApp::GetBuildConfiguration(), nullptr);
+				}
+				else if(FPaths::IsUnderDirectory(Target.Path, FPaths::EngineDir()))
+				{
+					EditorTargetFileName = FTargetReceipt::GetDefaultPath(*FPaths::EngineDir(), *Target.Name, FPlatformProcess::GetBinariesSubdirectory(), FApp::GetBuildConfiguration(), nullptr);
+				}
 				break;
 			}
 		}
