Add support for command line files to be found inside project and project plugins directories.

diff --git a/Engine/Source/Runtime/Launch/Private/LaunchEngineLoop.cpp b/Engine/Source/Runtime/Launch/Private/LaunchEngineLoop.cpp
--- a/Engine/Source/Runtime/Launch/Private/LaunchEngineLoop.cpp
+++ b/Engine/Source/Runtime/Launch/Private/LaunchEngineLoop.cpp
@@ -871,6 +871,119 @@ bool LaunchCheckForFileOverride(const TCHAR* CmdLine, bool& OutFileOverrideFound
 	return true;
 }
 
+#if !UE_BUILD_SHIPPING && WITH_EDITORONLY_DATA
+/**
+ * Process command line aliases
+ *
+ */
+void LaunchCheckForCommandLineAliases(bool& bChanged)
+{
+	bChanged = false;
+	if (FPaths::IsProjectFilePathSet())
+	{
+		FConfigFile Config;
+		if (FConfigCacheIni::LoadExternalIniFile(Config, TEXT("CommandLineAliases"), nullptr, *FPaths::Combine(FPaths::ProjectDir(), TEXT("Config")), nullptr))
+		{
+			if (FConfigSection* Section = Config.Find(TEXT("CommandLineAliases")))
+			{
+				TArray<FString> Tokens;
+				{
+					const TCHAR* Stream = FCommandLine::Get();
+					FString NextToken;
+					while (FParse::Token(Stream, NextToken, false))
+					{
+						Tokens.Add(NextToken);
+					}
+				}
+
+				for (FConfigSection::TIterator It(*Section); It; ++It)
+				{
+					FString Key = FString(TEXT("-")) + It.Key().ToString();
+					for (FString& Token : Tokens)
+					{
+						if (Token == Key)
+						{
+							Token = It.Value().GetValue();
+							bChanged = true;
+						}
+					}
+				}
+
+				if (bChanged)
+				{
+					FString NewCommandLine = FString::Join(Tokens, TEXT(" "));
+					FCommandLine::Set(*NewCommandLine);
+				}
+			}
+		}
+	}
+}
+
+/**
+ * Look for command line file
+ *
+ */
+void LaunchCheckForCmdLineFile(const TCHAR* CmdLine, bool& bChanged)
+{
+	bChanged = false;
+
+	FString CmdLineFile;
+	if (FParse::Value(FCommandLine::Get(), TEXT("-CmdLineFile="), CmdLineFile))
+	{
+		if (CmdLineFile.EndsWith(TEXT(".txt")))
+		{
+			auto TryProcessFile = [&bChanged](const FString& InFilePath)
+			{
+				FString FileCmds;
+				if (FFileHelper::LoadFileToString(FileCmds, *InFilePath))
+				{
+					FileCmds = FString(TEXT(" ")) + FileCmds.TrimStartAndEnd();
+					if (FileCmds.Len() > 1)
+					{
+						UE_LOG(LogInit, Log, TEXT("Appending commandline from file: %s, %s"), *InFilePath, *FileCmds);
+						FCommandLine::Append(*FileCmds);
+						bChanged = true;
+					}
+
+					return true;
+				}
+
+				return false;
+			};
+
+			bool bFoundFile = TryProcessFile(CmdLineFile);
+			if (!bFoundFile && FPaths::IsProjectFilePathSet())
+			{
+				const FString ProjectDir = FPaths::ProjectDir();
+				bFoundFile = TryProcessFile(ProjectDir + CmdLineFile);
+				if (!bFoundFile)
+				{
+					const FString ProjectPluginsDir = ProjectDir + TEXT("Plugins/");
+					TArray<FString> DirectoryNames;
+					IFileManager::Get().IterateDirectory(*ProjectPluginsDir, [&](const TCHAR* FilenameOrDirectory, bool bIsDirectory) -> bool
+					{
+						if (bIsDirectory && TryProcessFile(FString(FilenameOrDirectory) + TEXT("/") + CmdLineFile))
+						{
+							bFoundFile = true;
+							return false;
+						}
+						return true;
+					});
+				}
+			}
+
+			if (!bFoundFile)
+			{
+				UE_LOG(LogInit, Warning, TEXT("Failed to load commandline file '%s'."), *CmdLineFile);
+			}
+		}
+		else
+		{
+			UE_LOG(LogInit, Warning, TEXT("Can only load commandline files ending with .txt, can't load: %s"), *CmdLineFile);
+		}
+	}
+}
+#endif
 
 bool LaunchHasIncompleteGameName()
 {
@@ -1421,6 +1534,27 @@ int32 FEngineLoop::PreInitPreStartupScreen(const TCHAR* CmdLine)
 #endif
 	}
 
+	// Avoiding potential exploits by not exposing command line overrides in the shipping games.
+#if !UE_BUILD_SHIPPING && WITH_EDITORONLY_DATA
+	{
+		SCOPED_BOOT_TIMING("Command Line Adjustments");
+
+		bool bChanged = false;
+		LaunchCheckForCommandLineAliases(bChanged);
+		if (bChanged)
+		{
+			CmdLine = FCommandLine::Get();
+		}
+
+		bChanged = false;
+		LaunchCheckForCmdLineFile(CmdLine, bChanged);
+		if (bChanged)
+		{
+			CmdLine = FCommandLine::Get();
+		}
+	}
+#endif
+
 	// allow the command line to override the platform file singleton
 	bool bFileOverrideFound = false;
 	{
@@ -4964,37 +5098,6 @@ bool FEngineLoop::AppInit( )
 
 	// Avoiding potential exploits by not exposing command line overrides in the shipping games.
 #if !UE_BUILD_SHIPPING && WITH_EDITORONLY_DATA
-	FString CmdLineFile;
-
-	if (FParse::Value(FCommandLine::Get(), TEXT("-CmdLineFile="), CmdLineFile))
-	{
-		if (CmdLineFile.EndsWith(TEXT(".txt")))
-		{
-			FString FileCmds;
-
-			if (FFileHelper::LoadFileToString(FileCmds, *CmdLineFile))
-			{
-				FileCmds = FString(TEXT(" ")) + FileCmds.TrimStartAndEnd();
-
-				if (FileCmds.Len() > 1)
-				{
-					UE_LOG(LogInit, Log, TEXT("Appending commandline from file:%s"), *FileCmds);
-
-					FCommandLine::Append(*FileCmds);
-				}
-			}
-			else
-			{
-				UE_LOG(LogInit, Warning, TEXT("Failed to load commandline file '%s'."), *CmdLineFile);
-			}
-		}
-		else
-		{
-			UE_LOG(LogInit, Warning, TEXT("Can only load commandline files ending with .txt, can't load: %s"), *CmdLineFile);
-		}
-	}
-
-
 	// Retrieve additional command line arguments from environment variable.
 	FString Env = FPlatformMisc::GetEnvironmentVariable(TEXT("UE-CmdLineArgs")).TrimStart();
 	if (Env.Len())
