Fix compile error
#rnx

diff --git a/Engine/Source/Runtime/Launch/Private/LaunchEngineLoop.cpp b/Engine/Source/Runtime/Launch/Private/LaunchEngineLoop.cpp
--- a/Engine/Source/Runtime/Launch/Private/LaunchEngineLoop.cpp
+++ b/Engine/Source/Runtime/Launch/Private/LaunchEngineLoop.cpp
@@ -882,7 +882,7 @@ void LaunchCheckForCommandLineAliases(bool& bChanged)
 	if (FPaths::IsProjectFilePathSet())
 	{
 		FConfigFile Config;
-		if (FConfigCacheIni::LoadExternalIniFile(Config, TEXT("CommandLineAliases"), nullptr, *FPaths::Combine(FPaths::ProjectDir(), TEXT("Config")), nullptr))
+		if (FConfigCacheIni::LoadExternalIniFile(Config, TEXT("CommandLineAliases"), nullptr, *FPaths::Combine(FPaths::ProjectDir(), TEXT("Config")), false))
 		{
 			if (FConfigSection* Section = Config.Find(TEXT("CommandLineAliases")))
 			{
