Explicitly load SlateCore so that it can work with Live Coding

#rnx

diff --git a/Engine/Source/Runtime/Launch/Private/LaunchEngineLoop.cpp b/Engine/Source/Runtime/Launch/Private/LaunchEngineLoop.cpp
--- a/Engine/Source/Runtime/Launch/Private/LaunchEngineLoop.cpp
+++ b/Engine/Source/Runtime/Launch/Private/LaunchEngineLoop.cpp
@@ -3225,6 +3225,7 @@ bool FEngineLoop::LoadStartupCoreModules()
 	SlowTask.EnterProgressFrame(10);
 	if ( !IsRunningDedicatedServer() )
 	{
+		FModuleManager::Get().LoadModule("SlateCore");
 		FModuleManager::Get().LoadModule("Slate");
 
 #if !UE_BUILD_SHIPPING
