Added ability to check OS level notification status on Android.
#rnx

diff --git a/Engine/Build/Android/Java/src/com/epicgames/ue4/GameActivity.java.template b/Engine/Build/Android/Java/src/com/epicgames/ue4/GameActivity.java.template
--- a/Engine/Build/Android/Java/src/com/epicgames/ue4/GameActivity.java.template
+++ b/Engine/Build/Android/Java/src/com/epicgames/ue4/GameActivity.java.template
@@ -206,6 +206,7 @@ import android.graphics.Color;
 import android.widget.TextView;
 import android.widget.TextView.OnEditorActionListener;
 import android.content.BroadcastReceiver;
+import android.support.v4.app.NotificationManagerCompat;
 import android.support.v4.content.FileProvider;
 //$${gameActivityImportAdditions}$$
 
@@ -5129,6 +5130,11 @@ public class GameActivity extends $${gameActivitySuperClass}$$ implements Surfac
 		}
 	}
 
+	public boolean AndroidThunkJava_IsAllowedRemoteNotifications()
+	{
+		return NotificationManagerCompat.from(getApplicationContext()).areNotificationsEnabled();
+	}
+
 	public void AndroidThunkJava_SetDesiredViewSize(int width, int height)
 	{
 		if (width == DesiredHolderWidth && height == DesiredHolderHeight)
diff --git a/Engine/Source/Runtime/Core/Private/Android/AndroidPlatformMisc.cpp b/Engine/Source/Runtime/Core/Private/Android/AndroidPlatformMisc.cpp
--- a/Engine/Source/Runtime/Core/Private/Android/AndroidPlatformMisc.cpp
+++ b/Engine/Source/Runtime/Core/Private/Android/AndroidPlatformMisc.cpp
@@ -1319,6 +1319,7 @@ bool FAndroidMisc::IsStandaloneStereoOnlyDevice()
 
 extern void AndroidThunkCpp_RegisterForRemoteNotifications();
 extern void AndroidThunkCpp_UnregisterForRemoteNotifications();
+extern bool AndroidThunkCpp_IsAllowedRemoteNotifications();
 
 void FAndroidMisc::RegisterForRemoteNotifications()
 {
@@ -1334,6 +1335,15 @@ void FAndroidMisc::UnregisterForRemoteNotifications()
 #endif
 }
 
+bool FAndroidMisc::IsAllowedRemoteNotifications()
+{
+#if USE_ANDROID_JNI
+	return AndroidThunkCpp_IsAllowedRemoteNotifications();
+#else
+	return false;
+#endif
+}
+
 TArray<uint8> FAndroidMisc::GetSystemFontBytes()
 {
 #if USE_ANDROID_FILE
diff --git a/Engine/Source/Runtime/Core/Public/Android/AndroidPlatformMisc.h b/Engine/Source/Runtime/Core/Public/Android/AndroidPlatformMisc.h
--- a/Engine/Source/Runtime/Core/Public/Android/AndroidPlatformMisc.h
+++ b/Engine/Source/Runtime/Core/Public/Android/AndroidPlatformMisc.h
@@ -150,6 +150,7 @@ public:
 
 	static void RegisterForRemoteNotifications();
 	static void UnregisterForRemoteNotifications();
+	static bool IsAllowedRemoteNotifications();
 
 	/** @return Memory representing a true type or open type font provided by the platform as a default font for unreal to consume; empty array if the default font failed to load. */
 	static TArray<uint8> GetSystemFontBytes();
diff --git a/Engine/Source/Runtime/Launch/Private/Android/AndroidJNI.cpp b/Engine/Source/Runtime/Launch/Private/Android/AndroidJNI.cpp
--- a/Engine/Source/Runtime/Launch/Private/Android/AndroidJNI.cpp
+++ b/Engine/Source/Runtime/Launch/Private/Android/AndroidJNI.cpp
@@ -122,6 +122,7 @@ void FJavaWrapper::FindClassesAndMethods(JNIEnv* Env)
 	// this is optional - only inserted if GCM plugin enabled
 	AndroidThunkJava_RegisterForRemoteNotifications = FindMethod(Env, GameActivityClassID, "AndroidThunkJava_RegisterForRemoteNotifications", "()V", true);
 	AndroidThunkJava_UnregisterForRemoteNotifications = FindMethod(Env, GameActivityClassID, "AndroidThunkJava_UnregisterForRemoteNotifications", "()V", true);
+	AndroidThunkJava_IsAllowedRemoteNotifications = FindMethod(Env, GameActivityClassID, "AndroidThunkJava_IsAllowedRemoteNotifications", "()Z", true);
 
 	// get field IDs for InputDeviceInfo class members
 	InputDeviceInfoClass = FindClassGlobalRef(Env, "com/epicgames/ue4/GameActivity$InputDeviceInfo", bIsOptional);
@@ -376,6 +377,7 @@ jmethodID FJavaWrapper::AndroidThunkJava_GetMetaDataString;
 jmethodID FJavaWrapper::AndroidThunkJava_IsOculusMobileApplication;
 jmethodID FJavaWrapper::AndroidThunkJava_RegisterForRemoteNotifications;
 jmethodID FJavaWrapper::AndroidThunkJava_UnregisterForRemoteNotifications;
+jmethodID FJavaWrapper::AndroidThunkJava_IsAllowedRemoteNotifications;
 jmethodID FJavaWrapper::AndroidThunkJava_ShowHiddenAlertDialog;
 jmethodID FJavaWrapper::AndroidThunkJava_LocalNotificationScheduleAtTime;
 jmethodID FJavaWrapper::AndroidThunkJava_LocalNotificationClearAll;
@@ -800,6 +802,20 @@ void AndroidThunkCpp_UnregisterForRemoteNotifications()
 	}
 }
 
+bool AndroidThunkCpp_IsAllowedRemoteNotifications()
+{
+	bool Result = false;
+	if (FJavaWrapper::AndroidThunkJava_UnregisterForRemoteNotifications)
+	{
+		if (JNIEnv* Env = FAndroidApplication::GetJavaEnv())
+		{
+			Result = FJavaWrapper::CallBooleanMethod(Env, FJavaWrapper::GameActivityThis, FJavaWrapper::AndroidThunkJava_IsAllowedRemoteNotifications);
+		}
+	}
+
+	return Result;
+}
+
 void AndroidThunkCpp_ShowConsoleWindow()
 {
 	if (JNIEnv* Env = FAndroidApplication::GetJavaEnv())
diff --git a/Engine/Source/Runtime/Launch/Public/Android/AndroidJNI.h b/Engine/Source/Runtime/Launch/Public/Android/AndroidJNI.h
--- a/Engine/Source/Runtime/Launch/Public/Android/AndroidJNI.h
+++ b/Engine/Source/Runtime/Launch/Public/Android/AndroidJNI.h
@@ -102,6 +102,7 @@ public:
 	// Optionally added if GCM plugin (or other remote notification system) enabled
 	static jmethodID AndroidThunkJava_RegisterForRemoteNotifications;
 	static jmethodID AndroidThunkJava_UnregisterForRemoteNotifications;
+	static jmethodID AndroidThunkJava_IsAllowedRemoteNotifications;
 
 	// In app purchase functionality
 	static jclass JavaStringClass;
