Merging Version.h from Dev-EngineMerge to Dev-Main.

#rnx

diff --git a/Engine/Source/Runtime/Launch/Resources/Version.h b/Engine/Source/Runtime/Launch/Resources/Version.h
--- a/Engine/Source/Runtime/Launch/Resources/Version.h
+++ b/Engine/Source/Runtime/Launch/Resources/Version.h
@@ -1,4 +1,4 @@
-// Copyright 1998-2019 Epic Games, Inc. All Rights Reserved.
+// Copyright Epic Games, Inc. All Rights Reserved.
 
 #pragma once
 
@@ -65,7 +65,7 @@
 
 // Various strings used for engine resources
 #define EPIC_COMPANY_NAME  "Epic Games, Inc."
-#define EPIC_COPYRIGHT_STRING "Copyright 1998-2019 Epic Games, Inc. All Rights Reserved."
+#define EPIC_COPYRIGHT_STRING "Copyright Epic Games, Inc. All Rights Reserved."
 #define EPIC_PRODUCT_NAME "Unreal Engine"
 #define EPIC_PRODUCT_IDENTIFIER "UnrealEngine"
 
