General process lifetime tracing for Unreal Insights visualization. Added bookmarks:

* Tick loop end
* WinMain.Enter / WinMain.Exit
* UnloadModulesAtShutdown
* EditorExit

(peafour-cherrypick of 6543971 by [REDACTED_PERSON])

diff --git a/Engine/Source/Editor/UnrealEd/Private/UnrealEdGlobals.cpp b/Engine/Source/Editor/UnrealEd/Private/UnrealEdGlobals.cpp
--- a/Engine/Source/Editor/UnrealEd/Private/UnrealEdGlobals.cpp
+++ b/Engine/Source/Editor/UnrealEd/Private/UnrealEdGlobals.cpp
@@ -171,6 +171,8 @@ int32 EditorInit( IEngineLoop& EngineLoop )
 
 void EditorExit()
 {
+	TRACE_CPUPROFILER_EVENT_SCOPE_TEXT(TEXT("EditorExit"));
+
 	GLevelEditorModeTools().SetDefaultMode(FBuiltinEditorModes::EM_Default);
 	GLevelEditorModeTools().DeactivateAllModes(); // this also activates the default mode
 
diff --git a/Engine/Source/Runtime/Core/Private/Modules/ModuleManager.cpp b/Engine/Source/Runtime/Core/Private/Modules/ModuleManager.cpp
--- a/Engine/Source/Runtime/Core/Private/Modules/ModuleManager.cpp
+++ b/Engine/Source/Runtime/Core/Private/Modules/ModuleManager.cpp
@@ -678,6 +678,8 @@ void FModuleManager::UnloadModulesAtShutdown()
 {
 	ensure(IsInGameThread());
 
+	TRACE_CPUPROFILER_EVENT_SCOPE_TEXT(TEXT("UnloadModulesAtShutdown"));
+
 	struct FModulePair
 	{
 		FName ModuleName;
diff --git a/Engine/Source/Runtime/Launch/Private/Launch.cpp b/Engine/Source/Runtime/Launch/Private/Launch.cpp
--- a/Engine/Source/Runtime/Launch/Private/Launch.cpp
+++ b/Engine/Source/Runtime/Launch/Private/Launch.cpp
@@ -175,6 +175,8 @@ int32 GuardedMain( const TCHAR* CmdLine )
 		EngineTick();
 	}
 
+	TRACE_BOOKMARK(TEXT("Tick loop end"));
+
 #if WITH_EDITOR
 	if( GIsEditor )
 	{
diff --git a/Engine/Source/Runtime/Launch/Private/Windows/LaunchWindows.cpp b/Engine/Source/Runtime/Launch/Private/Windows/LaunchWindows.cpp
--- a/Engine/Source/Runtime/Launch/Private/Windows/LaunchWindows.cpp
+++ b/Engine/Source/Runtime/Launch/Private/Windows/LaunchWindows.cpp
@@ -208,6 +208,8 @@ bool ProcessCommandLine()
 
 int32 WINAPI WinMain( _In_ HINSTANCE hInInstance, _In_opt_ HINSTANCE hPrevInstance, _In_ char*, _In_ int32 nCmdShow )
 {
+	TRACE_BOOKMARK(TEXT("WinMain.Enter"));
+
 	// Setup common Windows settings
 	SetupWindowsEnvironment();
 
@@ -313,6 +315,8 @@ int32 WINAPI WinMain( _In_ HINSTANCE hInInstance, _In_opt_ HINSTANCE hPrevInstan
 		Sleep(INFINITE);
 	}
 
+	TRACE_BOOKMARK(TEXT("WinMain.Exit"));
+
 	return ErrorLevel;
 }
 
