TAutoPtr and TScopedPointer smart pointers removed.

diff --git a/Engine/Plugins/FX/Niagara/Source/NiagaraShader/Public/NiagaraShared.h b/Engine/Plugins/FX/Niagara/Source/NiagaraShader/Public/NiagaraShared.h
--- a/Engine/Plugins/FX/Niagara/Source/NiagaraShader/Public/NiagaraShared.h
+++ b/Engine/Plugins/FX/Niagara/Source/NiagaraShader/Public/NiagaraShared.h
@@ -11,7 +11,6 @@
 #include "Misc/Guid.h"
 #include "Engine/EngineTypes.h"
 #include "Templates/RefCounting.h"
-#include "Templates/ScopedPointer.h"
 #include "Misc/SecureHash.h"
 #include "RHI.h"
 #include "RenderResource.h"
diff --git a/Engine/Source/Developer/CrashDebugHelper/Private/CrashDebugPDBCache.cpp b/Engine/Source/Developer/CrashDebugHelper/Private/CrashDebugPDBCache.cpp
--- a/Engine/Source/Developer/CrashDebugHelper/Private/CrashDebugPDBCache.cpp
+++ b/Engine/Source/Developer/CrashDebugHelper/Private/CrashDebugPDBCache.cpp
@@ -4,7 +4,6 @@
 #include "HAL/PlatformTime.h"
 #include "HAL/FileManager.h"
 #include "Misc/Parse.h"
-#include "Templates/ScopedPointer.h"
 #include "Misc/CommandLine.h"
 #include "Misc/Paths.h"
 #include "Misc/ConfigCacheIni.h"
diff --git a/Engine/Source/Developer/DerivedDataCache/Private/DerivedDataBackendAsyncPutWrapper.h b/Engine/Source/Developer/DerivedDataCache/Private/DerivedDataBackendAsyncPutWrapper.h
--- a/Engine/Source/Developer/DerivedDataCache/Private/DerivedDataBackendAsyncPutWrapper.h
+++ b/Engine/Source/Developer/DerivedDataCache/Private/DerivedDataBackendAsyncPutWrapper.h
@@ -8,7 +8,6 @@
 #include "ProfilingDebugging/CookStats.h"
 #include "DerivedDataCacheUsageStats.h"
 #include "Misc/ScopeLock.h"
-#include "Templates/ScopedPointer.h"
 #include "MemoryDerivedDataBackend.h"
 #include "Async/AsyncWork.h"
 
diff --git a/Engine/Source/Developer/DerivedDataCache/Private/HierarchicalDerivedDataBackend.h b/Engine/Source/Developer/DerivedDataCache/Private/HierarchicalDerivedDataBackend.h
--- a/Engine/Source/Developer/DerivedDataCache/Private/HierarchicalDerivedDataBackend.h
+++ b/Engine/Source/Developer/DerivedDataCache/Private/HierarchicalDerivedDataBackend.h
@@ -6,7 +6,6 @@
 #include "DerivedDataBackendInterface.h"
 #include "ProfilingDebugging/CookStats.h"
 #include "DerivedDataCacheUsageStats.h"
-#include "Templates/ScopedPointer.h"
 #include "DerivedDataBackendAsyncPutWrapper.h"
 #include "Templates/UniquePtr.h"
 
diff --git a/Engine/Source/Developer/DerivedDataCache/Private/MemoryDerivedDataBackend.h b/Engine/Source/Developer/DerivedDataCache/Private/MemoryDerivedDataBackend.h
--- a/Engine/Source/Developer/DerivedDataCache/Private/MemoryDerivedDataBackend.h
+++ b/Engine/Source/Developer/DerivedDataCache/Private/MemoryDerivedDataBackend.h
@@ -8,7 +8,6 @@
 #include "ProfilingDebugging/CookStats.h"
 #include "DerivedDataCacheUsageStats.h"
 #include "Misc/ScopeLock.h"
-#include "Templates/ScopedPointer.h"
 
 class Error;
 
diff --git a/Engine/Source/Developer/DerivedDataCache/Private/PakFileDerivedDataBackend.h b/Engine/Source/Developer/DerivedDataCache/Private/PakFileDerivedDataBackend.h
--- a/Engine/Source/Developer/DerivedDataCache/Private/PakFileDerivedDataBackend.h
+++ b/Engine/Source/Developer/DerivedDataCache/Private/PakFileDerivedDataBackend.h
@@ -9,7 +9,6 @@
 #include "ProfilingDebugging/CookStats.h"
 #include "DerivedDataCacheUsageStats.h"
 #include "Misc/ScopeLock.h"
-#include "Templates/ScopedPointer.h"
 #include "Misc/FileHelper.h"
 #include "Serialization/MemoryReader.h"
 #include "Serialization/MemoryWriter.h"
diff --git a/Engine/Source/Developer/DesktopPlatform/Private/Windows/WindowsRegistry.h b/Engine/Source/Developer/DesktopPlatform/Private/Windows/WindowsRegistry.h
--- a/Engine/Source/Developer/DesktopPlatform/Private/Windows/WindowsRegistry.h
+++ b/Engine/Source/Developer/DesktopPlatform/Private/Windows/WindowsRegistry.h
@@ -6,7 +6,6 @@
 #include "Containers/Array.h"
 #include "Containers/Map.h"
 #include "Containers/UnrealString.h"
-#include "Templates/ScopedPointer.h"
 #include "Windows/WindowsHWrapper.h"
 #include "Templates/UniquePtr.h"
 
diff --git a/Engine/Source/Developer/MeshBoneReduction/Private/MeshBoneReduction.cpp b/Engine/Source/Developer/MeshBoneReduction/Private/MeshBoneReduction.cpp
--- a/Engine/Source/Developer/MeshBoneReduction/Private/MeshBoneReduction.cpp
+++ b/Engine/Source/Developer/MeshBoneReduction/Private/MeshBoneReduction.cpp
@@ -7,7 +7,6 @@
 #include "Engine/SkeletalMesh.h"
 #include "Components/SkinnedMeshComponent.h"
 #include "UObject/UObjectHash.h"
-#include "Templates/ScopedPointer.h"
 #include "ComponentReregisterContext.h"
 #include "Templates/UniquePtr.h"
 #include "Rendering/SkeletalMeshModel.h"
diff --git a/Engine/Source/Developer/MeshSimplifier/Private/QuadricMeshReduction.cpp b/Engine/Source/Developer/MeshSimplifier/Private/QuadricMeshReduction.cpp
--- a/Engine/Source/Developer/MeshSimplifier/Private/QuadricMeshReduction.cpp
+++ b/Engine/Source/Developer/MeshSimplifier/Private/QuadricMeshReduction.cpp
@@ -2,7 +2,6 @@
 
 #include "CoreMinimal.h"
 #include "Modules/ModuleManager.h"
-#include "Templates/ScopedPointer.h"
 #include "MeshUtilities.h"
 #include "MeshBuild.h"
 #include "MeshSimplify.h"
diff --git a/Engine/Source/Developer/PakFileUtilities/Private/PakFileUtilities.cpp b/Engine/Source/Developer/PakFileUtilities/Private/PakFileUtilities.cpp
--- a/Engine/Source/Developer/PakFileUtilities/Private/PakFileUtilities.cpp
+++ b/Engine/Source/Developer/PakFileUtilities/Private/PakFileUtilities.cpp
@@ -3182,7 +3182,7 @@ bool GenerateHashesFromPak(const TCHAR* InPakFilename, const TCHAR* InDestPakFil
 				    EntryInfo.Serialize(PakReader, PakFile.GetInfo().Version);
 				    if (EntryInfo == Entry)
 				    {
-					    // TAutoPtr<FArchive> FileHandle(IFileManager::Get().CreateFileWriter(*DestFilename));
+					    // TUniquePtr<FArchive> FileHandle(IFileManager::Get().CreateFileWriter(*DestFilename));
 					    TArray<uint8> Bytes;
 					    FMemoryWriter MemoryFile(Bytes);
 					    FArchive* FileHandle = &MemoryFile;
diff --git a/Engine/Source/Developer/Profiler/Private/ProfilerManager.cpp b/Engine/Source/Developer/Profiler/Private/ProfilerManager.cpp
--- a/Engine/Source/Developer/Profiler/Private/ProfilerManager.cpp
+++ b/Engine/Source/Developer/Profiler/Private/ProfilerManager.cpp
@@ -1,7 +1,6 @@
 // Copyright 1998-2019 Epic Games, Inc. All Rights Reserved.
 
 #include "ProfilerManager.h"
-#include "Templates/ScopedPointer.h"
 #include "Modules/ModuleManager.h"
 #include "IProfilerClientModule.h"
 #include "Stats/StatsFile.h"
diff --git a/Engine/Source/Developer/Profiler/Private/ProfilerRawStatsForThreadView.cpp b/Engine/Source/Developer/Profiler/Private/ProfilerRawStatsForThreadView.cpp
--- a/Engine/Source/Developer/Profiler/Private/ProfilerRawStatsForThreadView.cpp
+++ b/Engine/Source/Developer/Profiler/Private/ProfilerRawStatsForThreadView.cpp
@@ -2,7 +2,6 @@
 
 #include "ProfilerRawStatsForThreadView.h"
 #include "HAL/FileManager.h"
-#include "Templates/ScopedPointer.h"
 #include "Serialization/MemoryReader.h"
 #include "ProfilerDataProvider.h"
 
diff --git a/Engine/Source/Developer/TargetPlatform/Private/TargetPlatformManagerModule.cpp b/Engine/Source/Developer/TargetPlatform/Private/TargetPlatformManagerModule.cpp
--- a/Engine/Source/Developer/TargetPlatform/Private/TargetPlatformManagerModule.cpp
+++ b/Engine/Source/Developer/TargetPlatform/Private/TargetPlatformManagerModule.cpp
@@ -5,7 +5,6 @@
 #include "Misc/CommandLine.h"
 #include "Misc/Paths.h"
 #include "Misc/OutputDeviceRedirector.h"
-#include "Templates/ScopedPointer.h"
 #include "Stats/Stats.h"
 #include "Misc/ScopedSlowTask.h"
 #include "Misc/MonitoredProcess.h"
diff --git a/Engine/Source/Developer/TraceInsights/Private/Insights/InsightsManager.cpp b/Engine/Source/Developer/TraceInsights/Private/Insights/InsightsManager.cpp
--- a/Engine/Source/Developer/TraceInsights/Private/Insights/InsightsManager.cpp
+++ b/Engine/Source/Developer/TraceInsights/Private/Insights/InsightsManager.cpp
@@ -4,7 +4,6 @@
 
 #include "Framework/Application/SlateApplication.h"
 #include "Modules/ModuleManager.h"
-#include "Templates/ScopedPointer.h"
 #include "Templates/UniquePtr.h"
 
 // Insights
diff --git a/Engine/Source/Developer/TraceInsights/Private/Insights/IoProfilerManager.cpp b/Engine/Source/Developer/TraceInsights/Private/Insights/IoProfilerManager.cpp
--- a/Engine/Source/Developer/TraceInsights/Private/Insights/IoProfilerManager.cpp
+++ b/Engine/Source/Developer/TraceInsights/Private/Insights/IoProfilerManager.cpp
@@ -3,7 +3,6 @@
 #include "IoProfilerManager.h"
 
 #include "Modules/ModuleManager.h"
-#include "Templates/ScopedPointer.h"
 #include "Templates/UniquePtr.h"
 #include "TraceServices/AnalysisService.h"
 #include "TraceServices/SessionService.h"
diff --git a/Engine/Source/Developer/TraceInsights/Private/Insights/TimingProfilerManager.cpp b/Engine/Source/Developer/TraceInsights/Private/Insights/TimingProfilerManager.cpp
--- a/Engine/Source/Developer/TraceInsights/Private/Insights/TimingProfilerManager.cpp
+++ b/Engine/Source/Developer/TraceInsights/Private/Insights/TimingProfilerManager.cpp
@@ -3,7 +3,6 @@
 #include "TimingProfilerManager.h"
 
 #include "Modules/ModuleManager.h"
-#include "Templates/ScopedPointer.h"
 #include "Templates/UniquePtr.h"
 #include "TraceServices/AnalysisService.h"
 #include "TraceServices/SessionService.h"
diff --git a/Engine/Source/Editor/GraphEditor/Public/SNodePanel.h b/Engine/Source/Editor/GraphEditor/Public/SNodePanel.h
--- a/Engine/Source/Editor/GraphEditor/Public/SNodePanel.h
+++ b/Engine/Source/Editor/GraphEditor/Public/SNodePanel.h
@@ -20,7 +20,6 @@
 #include "Styling/CoreStyle.h"
 #include "Framework/Commands/InputChord.h"
 #include "GraphEditor.h"
-#include "Templates/ScopedPointer.h"
 #include "Layout/ArrangedChildren.h"
 #include "Types/PaintArgs.h"
 #include "EditorStyleSet.h"
diff --git a/Engine/Source/Editor/UMGEditor/Private/Designer/SDesignSurface.h b/Engine/Source/Editor/UMGEditor/Private/Designer/SDesignSurface.h
--- a/Engine/Source/Editor/UMGEditor/Private/Designer/SDesignSurface.h
+++ b/Engine/Source/Editor/UMGEditor/Private/Designer/SDesignSurface.h
@@ -12,7 +12,6 @@
 #include "Animation/CurveSequence.h"
 #include "Widgets/DeclarativeSyntaxSupport.h"
 #include "Widgets/SCompoundWidget.h"
-#include "Templates/ScopedPointer.h"
 #include "SNodePanel.h"
 #include "Templates/UniquePtr.h"
 
diff --git a/Engine/Source/Editor/UnrealEd/Classes/Commandlets/CookCommandlet.h b/Engine/Source/Editor/UnrealEd/Classes/Commandlets/CookCommandlet.h
--- a/Engine/Source/Editor/UnrealEd/Classes/Commandlets/CookCommandlet.h
+++ b/Engine/Source/Editor/UnrealEd/Classes/Commandlets/CookCommandlet.h
@@ -10,7 +10,6 @@
 #include "UObject/ObjectMacros.h"
 #include "Misc/Guid.h"
 #include "UObject/WeakObjectPtr.h"
-#include "Templates/ScopedPointer.h"
 #include "Misc/PackageName.h"
 #include "Commandlets/Commandlet.h"
 #include "Templates/UniquePtr.h"
diff --git a/Engine/Source/Editor/UnrealEd/Private/Commandlets/GatherTextFromAssetsCommandlet.cpp b/Engine/Source/Editor/UnrealEd/Private/Commandlets/GatherTextFromAssetsCommandlet.cpp
--- a/Engine/Source/Editor/UnrealEd/Private/Commandlets/GatherTextFromAssetsCommandlet.cpp
+++ b/Engine/Source/Editor/UnrealEd/Private/Commandlets/GatherTextFromAssetsCommandlet.cpp
@@ -17,7 +17,6 @@
 #include "UObject/PackageFileSummary.h"
 #include "Framework/Commands/Commands.h"
 #include "Commandlets/GatherTextFromSourceCommandlet.h"
-#include "Templates/ScopedPointer.h"
 #include "AssetData.h"
 #include "Sound/DialogueWave.h"
 #include "ARFilter.h"
diff --git a/Engine/Source/Editor/UnrealEd/Private/Commandlets/GenerateTextLocalizationResourceCommandlet.cpp b/Engine/Source/Editor/UnrealEd/Private/Commandlets/GenerateTextLocalizationResourceCommandlet.cpp
--- a/Engine/Source/Editor/UnrealEd/Private/Commandlets/GenerateTextLocalizationResourceCommandlet.cpp
+++ b/Engine/Source/Editor/UnrealEd/Private/Commandlets/GenerateTextLocalizationResourceCommandlet.cpp
@@ -3,7 +3,6 @@
 #include "Commandlets/GenerateTextLocalizationResourceCommandlet.h"
 #include "HAL/FileManager.h"
 #include "Misc/Paths.h"
-#include "Templates/ScopedPointer.h"
 #include "Internationalization/TextLocalizationResource.h"
 #include "TextLocalizationResourceGenerator.h"
 #include "Templates/UniquePtr.h"
diff --git a/Engine/Source/Editor/UnrealEd/Private/Commandlets/StabilizeLocalizationKeys.cpp b/Engine/Source/Editor/UnrealEd/Private/Commandlets/StabilizeLocalizationKeys.cpp
--- a/Engine/Source/Editor/UnrealEd/Private/Commandlets/StabilizeLocalizationKeys.cpp
+++ b/Engine/Source/Editor/UnrealEd/Private/Commandlets/StabilizeLocalizationKeys.cpp
@@ -12,7 +12,6 @@
 #include "UObject/Linker.h"
 #include "Internationalization/InternationalizationManifest.h"
 #include "Internationalization/InternationalizationArchive.h"
-#include "Templates/ScopedPointer.h"
 #include "UserDefinedStructure/UserDefinedStructEditorData.h"
 #include "Engine/UserDefinedStruct.h"
 #include "Serialization/JsonInternationalizationManifestSerializer.h"
diff --git a/Engine/Source/Editor/UnrealEd/Private/DistanceFieldBuildNotification.cpp b/Engine/Source/Editor/UnrealEd/Private/DistanceFieldBuildNotification.cpp
--- a/Engine/Source/Editor/UnrealEd/Private/DistanceFieldBuildNotification.cpp
+++ b/Engine/Source/Editor/UnrealEd/Private/DistanceFieldBuildNotification.cpp
@@ -4,7 +4,6 @@
 #include "CoreMinimal.h"
 #include "Stats/Stats.h"
 #include "TickableEditorObject.h"
-#include "Templates/ScopedPointer.h"
 #include "DistanceFieldAtlas.h"
 #include "Framework/Notifications/NotificationManager.h"
 #include "Widgets/Notifications/SNotificationList.h"
diff --git a/Engine/Source/Editor/UnrealEd/Private/StaticLightingSystem/StaticLightingPrivate.h b/Engine/Source/Editor/UnrealEd/Private/StaticLightingSystem/StaticLightingPrivate.h
--- a/Engine/Source/Editor/UnrealEd/Private/StaticLightingSystem/StaticLightingPrivate.h
+++ b/Engine/Source/Editor/UnrealEd/Private/StaticLightingSystem/StaticLightingPrivate.h
@@ -11,7 +11,6 @@
 #include "CoreMinimal.h"
 #include "SlateFwd.h"
 #include "Engine/Level.h"
-#include "Templates/ScopedPointer.h"
 #include "StaticLighting.h"
 #include "LightingBuildOptions.h"
 #include "Templates/UniquePtr.h"
diff --git a/Engine/Source/Editor/UnrealEd/Private/UnrealEdPrivatePCH.h b/Engine/Source/Editor/UnrealEd/Private/UnrealEdPrivatePCH.h
--- a/Engine/Source/Editor/UnrealEd/Private/UnrealEdPrivatePCH.h
+++ b/Engine/Source/Editor/UnrealEd/Private/UnrealEdPrivatePCH.h
@@ -206,7 +206,6 @@
 #include "Async/TaskGraphInterfaces.h"
 #include "Serialization/BufferReader.h"
 #include "Misc/SecureHash.h"
-#include "Templates/ScopedPointer.h"
 #include "Containers/StaticArray.h"
 #include "Misc/Paths.h"
 #include "Misc/CoreStats.h"
diff --git a/Engine/Source/Editor/UnrealEd/Public/UnrealEdSharedPCH.h b/Engine/Source/Editor/UnrealEd/Public/UnrealEdSharedPCH.h
--- a/Engine/Source/Editor/UnrealEd/Public/UnrealEdSharedPCH.h
+++ b/Engine/Source/Editor/UnrealEd/Public/UnrealEdSharedPCH.h
@@ -199,7 +199,6 @@
 #include "Templates/RefCounting.h"
 #include "Async/TaskGraphInterfaces.h"
 #include "Misc/Paths.h"
-#include "Templates/ScopedPointer.h"
 #include "Serialization/CustomVersion.h"
 #include "Misc/OutputDeviceError.h"
 #include "Misc/ObjectThumbnail.h"
diff --git a/Engine/Source/Programs/UnrealFrontend/Private/Commands/StatsConvertCommand.cpp b/Engine/Source/Programs/UnrealFrontend/Private/Commands/StatsConvertCommand.cpp
--- a/Engine/Source/Programs/UnrealFrontend/Private/Commands/StatsConvertCommand.cpp
+++ b/Engine/Source/Programs/UnrealFrontend/Private/Commands/StatsConvertCommand.cpp
@@ -3,7 +3,6 @@
 #include "StatsConvertCommand.h"
 #include "Misc/CommandLine.h"
 #include "HAL/FileManager.h"
-#include "Templates/ScopedPointer.h"
 #include "Templates/UniquePtr.h"
 
 /** Helper class used to extract stats data into CSV file. */
diff --git a/Engine/Source/Programs/UnrealLightmass/Private/Lighting/Radiosity.cpp b/Engine/Source/Programs/UnrealLightmass/Private/Lighting/Radiosity.cpp
--- a/Engine/Source/Programs/UnrealLightmass/Private/Lighting/Radiosity.cpp
+++ b/Engine/Source/Programs/UnrealLightmass/Private/Lighting/Radiosity.cpp
@@ -3,7 +3,6 @@
 #include "Raster.h"
 #include "LightingSystem.h"
 #include "LightmassSwarm.h"
-#include "Templates/ScopedPointer.h"
 #include "HAL/RunnableThread.h"
 #include "HAL/PlatformProcess.h"
 #include "TextureMappingSetup.h"
diff --git a/Engine/Source/Programs/UnrealLightmass/Private/Lighting/TextureMapping.cpp b/Engine/Source/Programs/UnrealLightmass/Private/Lighting/TextureMapping.cpp
--- a/Engine/Source/Programs/UnrealLightmass/Private/Lighting/TextureMapping.cpp
+++ b/Engine/Source/Programs/UnrealLightmass/Private/Lighting/TextureMapping.cpp
@@ -4,7 +4,6 @@
 #include "Raster.h"
 #include "LightingSystem.h"
 #include "LightmassSwarm.h"
-#include "Templates/ScopedPointer.h"
 #include "HAL/RunnableThread.h"
 #include "HAL/PlatformProcess.h"
 #include "TextureMappingSetup.h"
diff --git a/Engine/Source/Programs/UnrealLightmass/Private/LightmassCore/LMCore.cpp b/Engine/Source/Programs/UnrealLightmass/Private/LightmassCore/LMCore.cpp
--- a/Engine/Source/Programs/UnrealLightmass/Private/LightmassCore/LMCore.cpp
+++ b/Engine/Source/Programs/UnrealLightmass/Private/LightmassCore/LMCore.cpp
@@ -5,7 +5,6 @@
 #include "HAL/PlatformProcess.h"
 #include "HAL/FileManager.h"
 #include "HAL/PlatformTime.h"
-#include "Templates/ScopedPointer.h"
 #include "UnrealLightmass.h"
 #include "Misc/Paths.h"
 #include "Templates/UniquePtr.h"
diff --git a/Engine/Source/Runtime/ApplicationCore/Private/Windows/WindowsApplication.cpp b/Engine/Source/Runtime/ApplicationCore/Private/Windows/WindowsApplication.cpp
--- a/Engine/Source/Runtime/ApplicationCore/Private/Windows/WindowsApplication.cpp
+++ b/Engine/Source/Runtime/ApplicationCore/Private/Windows/WindowsApplication.cpp
@@ -2,7 +2,6 @@
 
 #include "Windows/WindowsApplication.h"
 #include "Containers/StringConv.h"
-#include "Templates/ScopedPointer.h"
 #include "CoreGlobals.h"
 #include "Internationalization/Text.h"
 #include "Misc/ConfigCacheIni.h"
diff --git a/Engine/Source/Runtime/CookedIterativeFile/Private/CookedIterativeNetworkFile.cpp b/Engine/Source/Runtime/CookedIterativeFile/Private/CookedIterativeNetworkFile.cpp
--- a/Engine/Source/Runtime/CookedIterativeFile/Private/CookedIterativeNetworkFile.cpp
+++ b/Engine/Source/Runtime/CookedIterativeFile/Private/CookedIterativeNetworkFile.cpp
@@ -1,7 +1,6 @@
 // Copyright 1998-2019 Epic Games, Inc. All Rights Reserved.
 
 #include "CookedIterativeNetworkFile.h"
-#include "Templates/ScopedPointer.h"
 #include "Misc/CommandLine.h"
 #include "Misc/ScopeLock.h"
 #include "Modules/ModuleManager.h"
diff --git a/Engine/Source/Runtime/Core/Private/GenericPlatform/GenericPlatformFile.cpp b/Engine/Source/Runtime/Core/Private/GenericPlatform/GenericPlatformFile.cpp
--- a/Engine/Source/Runtime/Core/Private/GenericPlatform/GenericPlatformFile.cpp
+++ b/Engine/Source/Runtime/Core/Private/GenericPlatform/GenericPlatformFile.cpp
@@ -2,7 +2,6 @@
 
 #include "GenericPlatform/GenericPlatformFile.h"
 #include "HAL/FileManager.h"
-#include "Templates/ScopedPointer.h"
 #include "Misc/Paths.h"
 #include "HAL/ThreadSafeCounter.h"
 #include "Stats/Stats.h"
diff --git a/Engine/Source/Runtime/Core/Private/GenericPlatform/GenericPlatformOutputDevices.cpp b/Engine/Source/Runtime/Core/Private/GenericPlatform/GenericPlatformOutputDevices.cpp
--- a/Engine/Source/Runtime/Core/Private/GenericPlatform/GenericPlatformOutputDevices.cpp
+++ b/Engine/Source/Runtime/Core/Private/GenericPlatform/GenericPlatformOutputDevices.cpp
@@ -4,7 +4,6 @@
 #include "HAL/PlatformOutputDevices.h"
 #include "CoreGlobals.h"
 #include "Misc/Parse.h"
-#include "Templates/ScopedPointer.h"
 #include "Misc/CommandLine.h"
 #include "Misc/Paths.h"
 #include "Misc/OutputDeviceMemory.h"
diff --git a/Engine/Source/Runtime/Core/Private/HAL/PlatformFileManager.cpp b/Engine/Source/Runtime/Core/Private/HAL/PlatformFileManager.cpp
--- a/Engine/Source/Runtime/Core/Private/HAL/PlatformFileManager.cpp
+++ b/Engine/Source/Runtime/Core/Private/HAL/PlatformFileManager.cpp
@@ -8,7 +8,6 @@
 #include "Misc/AssertionMacros.h"
 #include "GenericPlatform/GenericPlatformFile.h"
 #include "Modules/ModuleManager.h"
-#include "Templates/ScopedPointer.h"
 #include "HAL/IPlatformFileLogWrapper.h"
 #include "HAL/IPlatformFileProfilerWrapper.h"
 #include "HAL/IPlatformFileCachedWrapper.h"
diff --git a/Engine/Source/Runtime/Core/Private/Internationalization/TextLocalizationManager.cpp b/Engine/Source/Runtime/Core/Private/Internationalization/TextLocalizationManager.cpp
--- a/Engine/Source/Runtime/Core/Private/Internationalization/TextLocalizationManager.cpp
+++ b/Engine/Source/Runtime/Core/Private/Internationalization/TextLocalizationManager.cpp
@@ -10,7 +10,6 @@
 #include "HAL/FileManager.h"
 #include "HAL/LowLevelMemTracker.h"
 #include "Misc/Parse.h"
-#include "Templates/ScopedPointer.h"
 #include "Misc/ScopeLock.h"
 #include "Misc/CommandLine.h"
 #include "Internationalization/Culture.h"
diff --git a/Engine/Source/Runtime/Core/Private/Lumin/LuminPlatformOutputDevices.cpp b/Engine/Source/Runtime/Core/Private/Lumin/LuminPlatformOutputDevices.cpp
--- a/Engine/Source/Runtime/Core/Private/Lumin/LuminPlatformOutputDevices.cpp
+++ b/Engine/Source/Runtime/Core/Private/Lumin/LuminPlatformOutputDevices.cpp
@@ -5,7 +5,6 @@
 #include "Lumin/LuminOutputDeviceDebug.h"
 #include "CoreGlobals.h"
 #include "Misc/Parse.h"
-#include "Templates/ScopedPointer.h"
 #include "Misc/CommandLine.h"
 #include "Misc/Paths.h"
 #include "Misc/OutputDeviceMemory.h"
diff --git a/Engine/Source/Runtime/Core/Private/ProfilingDebugging/Apple/InstrumentsExternalProfiler.cpp b/Engine/Source/Runtime/Core/Private/ProfilingDebugging/Apple/InstrumentsExternalProfiler.cpp
--- a/Engine/Source/Runtime/Core/Private/ProfilingDebugging/Apple/InstrumentsExternalProfiler.cpp
+++ b/Engine/Source/Runtime/Core/Private/ProfilingDebugging/Apple/InstrumentsExternalProfiler.cpp
@@ -6,7 +6,6 @@
 #include "Misc/AssertionMacros.h"
 #include "ProfilingDebugging/ExternalProfiler.h"
 #include "Features/IModularFeatures.h"
-#include "Templates/ScopedPointer.h"
 #include "Templates/UniquePtr.h"
 #include "Apple/ApplePlatformDebugEvents.h"
 #include "Misc/Paths.h"
diff --git a/Engine/Source/Runtime/Core/Private/ProfilingDebugging/Windows/AQtimeExternalProfiler.cpp b/Engine/Source/Runtime/Core/Private/ProfilingDebugging/Windows/AQtimeExternalProfiler.cpp
--- a/Engine/Source/Runtime/Core/Private/ProfilingDebugging/Windows/AQtimeExternalProfiler.cpp
+++ b/Engine/Source/Runtime/Core/Private/ProfilingDebugging/Windows/AQtimeExternalProfiler.cpp
@@ -3,7 +3,6 @@
 #include "CoreTypes.h"
 #include "Misc/AssertionMacros.h"
 #include "ProfilingDebugging/ExternalProfiler.h"
-#include "Templates/ScopedPointer.h"
 #include "Features/IModularFeatures.h"
 #include "Windows/WindowsPlatformProcess.h"
 #include "Templates/UniquePtr.h"
diff --git a/Engine/Source/Runtime/Core/Private/ProfilingDebugging/Windows/VSPerfExternalProfiler.cpp b/Engine/Source/Runtime/Core/Private/ProfilingDebugging/Windows/VSPerfExternalProfiler.cpp
--- a/Engine/Source/Runtime/Core/Private/ProfilingDebugging/Windows/VSPerfExternalProfiler.cpp
+++ b/Engine/Source/Runtime/Core/Private/ProfilingDebugging/Windows/VSPerfExternalProfiler.cpp
@@ -5,7 +5,6 @@
 #include "Misc/AssertionMacros.h"
 #include "ProfilingDebugging/ExternalProfiler.h"
 #include "Features/IModularFeatures.h"
-#include "Templates/ScopedPointer.h"
 #include "Templates/UniquePtr.h"
 
 // Not all versions of Visual Studio include the profiler SDK headers
diff --git a/Engine/Source/Runtime/Core/Private/ProfilingDebugging/Windows/VTuneExternalProfiler.cpp b/Engine/Source/Runtime/Core/Private/ProfilingDebugging/Windows/VTuneExternalProfiler.cpp
--- a/Engine/Source/Runtime/Core/Private/ProfilingDebugging/Windows/VTuneExternalProfiler.cpp
+++ b/Engine/Source/Runtime/Core/Private/ProfilingDebugging/Windows/VTuneExternalProfiler.cpp
@@ -4,7 +4,6 @@
 #include "HAL/PlatformProcess.h"
 #include "Misc/AssertionMacros.h"
 #include "ProfilingDebugging/ExternalProfiler.h"
-#include "Templates/ScopedPointer.h"
 #include "Features/IModularFeatures.h"
 #include "Templates/UniquePtr.h"
 #include "Containers/Map.h"
diff --git a/Engine/Source/Runtime/Core/Private/Stats/StatsFile.cpp b/Engine/Source/Runtime/Core/Private/Stats/StatsFile.cpp
--- a/Engine/Source/Runtime/Core/Private/Stats/StatsFile.cpp
+++ b/Engine/Source/Runtime/Core/Private/Stats/StatsFile.cpp
@@ -6,7 +6,6 @@
 
 #include "Stats/StatsFile.h"
 #include "HAL/FileManager.h"
-#include "Templates/ScopedPointer.h"
 #include "Misc/Paths.h"
 #include "Internationalization/Internationalization.h"
 #include "Serialization/MemoryWriter.h"
diff --git a/Engine/Source/Runtime/Core/Private/Windows/WindowsPlatformCrashContext.cpp b/Engine/Source/Runtime/Core/Private/Windows/WindowsPlatformCrashContext.cpp
--- a/Engine/Source/Runtime/Core/Private/Windows/WindowsPlatformCrashContext.cpp
+++ b/Engine/Source/Runtime/Core/Private/Windows/WindowsPlatformCrashContext.cpp
@@ -19,7 +19,6 @@
 #include "Misc/ConfigCacheIni.h"
 #include "Misc/OutputDeviceRedirector.h"
 #include "Misc/OutputDeviceFile.h"
-#include "Templates/ScopedPointer.h"
 #include "Windows/WindowsPlatformStackWalk.h"
 #include "Windows/WindowsHWrapper.h"
 #include "Windows/AllowWindowsPlatformTypes.h"
diff --git a/Engine/Source/Runtime/Core/Public/Core.h b/Engine/Source/Runtime/Core/Public/Core.h
--- a/Engine/Source/Runtime/Core/Public/Core.h
+++ b/Engine/Source/Runtime/Core/Public/Core.h
@@ -123,7 +123,6 @@ MONOLITHIC_HEADER_BOILERPLATE()
 #include "Serialization/Archive.h"
 #include "Serialization/ArchiveProxy.h"
 #include "Serialization/NameAsStringProxyArchive.h"
-#include "Templates/ScopedPointer.h"
 #include "Templates/Sorting.h"
 #include "Containers/Array.h"
 #include "Containers/ScriptArray.h"
diff --git a/Engine/Source/Runtime/Core/Public/HAL/FileManagerGeneric.h b/Engine/Source/Runtime/Core/Public/HAL/FileManagerGeneric.h
--- a/Engine/Source/Runtime/Core/Public/HAL/FileManagerGeneric.h
+++ b/Engine/Source/Runtime/Core/Public/HAL/FileManagerGeneric.h
@@ -9,7 +9,6 @@
 #include "GenericPlatform/GenericPlatformFile.h"
 #include "HAL/PlatformFilemanager.h"
 #include "HAL/FileManager.h"
-#include "Templates/ScopedPointer.h"
 #include "Templates/UniquePtr.h"
 
 #ifndef PLATFORM_FILE_READER_BUFFER_SIZE
diff --git a/Engine/Source/Runtime/Core/Public/HAL/IPlatformFileLogWrapper.h b/Engine/Source/Runtime/Core/Public/HAL/IPlatformFileLogWrapper.h
--- a/Engine/Source/Runtime/Core/Public/HAL/IPlatformFileLogWrapper.h
+++ b/Engine/Source/Runtime/Core/Public/HAL/IPlatformFileLogWrapper.h
@@ -11,7 +11,6 @@
 #include "Misc/DateTime.h"
 #include "GenericPlatform/GenericPlatformFile.h"
 #include "HAL/PlatformTime.h"
-#include "Templates/ScopedPointer.h"
 #include "Misc/ScopeLock.h"
 #include "Templates/UniquePtr.h"
 
diff --git a/Engine/Source/Runtime/Core/Public/HAL/IPlatformFileProfilerWrapper.h b/Engine/Source/Runtime/Core/Public/HAL/IPlatformFileProfilerWrapper.h
--- a/Engine/Source/Runtime/Core/Public/HAL/IPlatformFileProfilerWrapper.h
+++ b/Engine/Source/Runtime/Core/Public/HAL/IPlatformFileProfilerWrapper.h
@@ -13,7 +13,6 @@
 #include "Misc/DateTime.h"
 #include "GenericPlatform/GenericPlatformFile.h"
 #include "HAL/PlatformTime.h"
-#include "Templates/ScopedPointer.h"
 #include "Misc/ScopeLock.h"
 #include "Templates/UniquePtr.h"
 
diff --git a/Engine/Source/Runtime/Core/Public/Templates/ScopedPointer.h b/Engine/Source/Runtime/Core/Public/Templates/ScopedPointer.h
deleted file mode 100644
diff --git a/Engine/Source/Runtime/CoreUObject/Private/CoreUObjectPrivatePCH.h b/Engine/Source/Runtime/CoreUObject/Private/CoreUObjectPrivatePCH.h
--- a/Engine/Source/Runtime/CoreUObject/Private/CoreUObjectPrivatePCH.h
+++ b/Engine/Source/Runtime/CoreUObject/Private/CoreUObjectPrivatePCH.h
@@ -180,7 +180,6 @@
 #include "Math/RandomStream.h"
 #include "Containers/List.h"
 #include "Serialization/CustomVersion.h"
-#include "Templates/ScopedPointer.h"
 #include "Misc/OutputDeviceError.h"
 #include "Misc/ObjectThumbnail.h"
 #include "Internationalization/GatherableTextData.h"
diff --git a/Engine/Source/Runtime/CoreUObject/Private/Misc/WorldCompositionUtility.cpp b/Engine/Source/Runtime/CoreUObject/Private/Misc/WorldCompositionUtility.cpp
--- a/Engine/Source/Runtime/CoreUObject/Private/Misc/WorldCompositionUtility.cpp
+++ b/Engine/Source/Runtime/CoreUObject/Private/Misc/WorldCompositionUtility.cpp
@@ -5,7 +5,6 @@
 =============================================================================*/
 #include "Misc/WorldCompositionUtility.h"
 #include "HAL/FileManager.h"
-#include "Templates/ScopedPointer.h"
 #include "UObject/PropertyPortFlags.h"
 #include "UObject/PackageFileSummary.h"
 #include "UObject/Linker.h"
diff --git a/Engine/Source/Runtime/CoreUObject/Public/UObject/Package.h b/Engine/Source/Runtime/CoreUObject/Public/UObject/Package.h
--- a/Engine/Source/Runtime/CoreUObject/Public/UObject/Package.h
+++ b/Engine/Source/Runtime/CoreUObject/Public/UObject/Package.h
@@ -8,7 +8,6 @@
 #include "UObject/Object.h"
 #include "Misc/Guid.h"
 #include "Misc/WorldCompositionUtility.h"
-#include "Templates/ScopedPointer.h"
 #include "Misc/OutputDeviceError.h"
 #include "Misc/ObjectThumbnail.h"
 #include "Serialization/CustomVersion.h"
diff --git a/Engine/Source/Runtime/Engine/Classes/Components/StaticMeshComponent.h b/Engine/Source/Runtime/Engine/Classes/Components/StaticMeshComponent.h
--- a/Engine/Source/Runtime/Engine/Classes/Components/StaticMeshComponent.h
+++ b/Engine/Source/Runtime/Engine/Classes/Components/StaticMeshComponent.h
@@ -8,7 +8,6 @@
 #include "Misc/Guid.h"
 #include "UObject/Class.h"
 #include "Engine/EngineTypes.h"
-#include "Templates/ScopedPointer.h"
 #include "Engine/TextureStreamingTypes.h"
 #include "Components/MeshComponent.h"
 #include "PackedNormal.h"
diff --git a/Engine/Source/Runtime/Engine/Classes/Engine/StaticMesh.h b/Engine/Source/Runtime/Engine/Classes/Engine/StaticMesh.h
--- a/Engine/Source/Runtime/Engine/Classes/Engine/StaticMesh.h
+++ b/Engine/Source/Runtime/Engine/Classes/Engine/StaticMesh.h
@@ -12,7 +12,6 @@
 #include "UObject/ScriptMacros.h"
 #include "Interfaces/Interface_AssetUserData.h"
 #include "RenderCommandFence.h"
-#include "Templates/ScopedPointer.h"
 #include "Components.h"
 #include "Interfaces/Interface_CollisionDataProvider.h"
 #include "Engine/MeshMerging.h"
diff --git a/Engine/Source/Runtime/Engine/Classes/Materials/Material.h b/Engine/Source/Runtime/Engine/Classes/Materials/Material.h
--- a/Engine/Source/Runtime/Engine/Classes/Materials/Material.h
+++ b/Engine/Source/Runtime/Engine/Classes/Materials/Material.h
@@ -7,7 +7,6 @@
 #include "Misc/Guid.h"
 #include "Engine/EngineTypes.h"
 #include "RenderCommandFence.h"
-#include "Templates/ScopedPointer.h"
 #include "Materials/MaterialInterface.h"
 #include "MaterialShared.h"
 #include "MaterialExpressionIO.h"
diff --git a/Engine/Source/Runtime/Engine/Classes/Materials/MaterialParameterCollection.h b/Engine/Source/Runtime/Engine/Classes/Materials/MaterialParameterCollection.h
--- a/Engine/Source/Runtime/Engine/Classes/Materials/MaterialParameterCollection.h
+++ b/Engine/Source/Runtime/Engine/Classes/Materials/MaterialParameterCollection.h
@@ -10,7 +10,6 @@
 #include "UObject/ObjectMacros.h"
 #include "UObject/Object.h"
 #include "Misc/Guid.h"
-#include "Templates/ScopedPointer.h"
 #include "UniformBuffer.h"
 #include "Templates/UniquePtr.h"
 #include "RenderCommandFence.h"
diff --git a/Engine/Source/Runtime/Engine/Private/EnginePrivatePCH.h b/Engine/Source/Runtime/Engine/Private/EnginePrivatePCH.h
--- a/Engine/Source/Runtime/Engine/Private/EnginePrivatePCH.h
+++ b/Engine/Source/Runtime/Engine/Private/EnginePrivatePCH.h
@@ -211,7 +211,6 @@
 #include "Misc/SecureHash.h"
 #include "Containers/StaticArray.h"
 #include "Async/Future.h"
-#include "Templates/ScopedPointer.h"
 #include "Serialization/ArchiveProxy.h"
 #include "UObject/DebugSerializationFlags.h"
 #include "Math/SHMath.h"
diff --git a/Engine/Source/Runtime/Engine/Private/Level.cpp b/Engine/Source/Runtime/Engine/Private/Level.cpp
--- a/Engine/Source/Runtime/Engine/Private/Level.cpp
+++ b/Engine/Source/Runtime/Engine/Private/Level.cpp
@@ -7,7 +7,6 @@ Level.cpp: Level-related functions
 #include "Engine/Level.h"
 #include "Misc/ScopedSlowTask.h"
 #include "UObject/RenderingObjectVersion.h"
-#include "Templates/ScopedPointer.h"
 #include "Templates/UnrealTemplate.h"
 #include "UObject/Package.h"
 #include "Serialization/AsyncLoading.h"
diff --git a/Engine/Source/Runtime/Engine/Private/RendererSupport.cpp b/Engine/Source/Runtime/Engine/Private/RendererSupport.cpp
--- a/Engine/Source/Runtime/Engine/Private/RendererSupport.cpp
+++ b/Engine/Source/Runtime/Engine/Private/RendererSupport.cpp
@@ -15,7 +15,6 @@
 #include "Shader.h"
 #include "VertexFactory.h"
 #include "SceneTypes.h"
-#include "Templates/ScopedPointer.h"
 #include "Materials/MaterialInterface.h"
 #include "MaterialShared.h"
 #include "Materials/Material.h"
diff --git a/Engine/Source/Runtime/Engine/Private/UnrealClient.cpp b/Engine/Source/Runtime/Engine/Private/UnrealClient.cpp
--- a/Engine/Source/Runtime/Engine/Private/UnrealClient.cpp
+++ b/Engine/Source/Runtime/Engine/Private/UnrealClient.cpp
@@ -8,7 +8,6 @@
 #include "EngineStats.h"
 #include "EngineGlobals.h"
 #include "RenderingThread.h"
-#include "Templates/ScopedPointer.h"
 #include "CanvasItem.h"
 #include "CanvasTypes.h"
 #include "Misc/ConfigCacheIni.h"
diff --git a/Engine/Source/Runtime/Engine/Public/DistanceFieldAtlas.h b/Engine/Source/Runtime/Engine/Public/DistanceFieldAtlas.h
--- a/Engine/Source/Runtime/Engine/Public/DistanceFieldAtlas.h
+++ b/Engine/Source/Runtime/Engine/Public/DistanceFieldAtlas.h
@@ -10,7 +10,6 @@
 #include "Containers/LockFreeList.h"
 #include "ProfilingDebugging/ResourceSize.h"
 #include "Engine/EngineTypes.h"
-#include "Templates/ScopedPointer.h"
 #include "UObject/GCObject.h"
 #include "RenderResource.h"
 #include "RenderingThread.h"
diff --git a/Engine/Source/Runtime/Engine/Public/EngineSharedPCH.h b/Engine/Source/Runtime/Engine/Public/EngineSharedPCH.h
--- a/Engine/Source/Runtime/Engine/Public/EngineSharedPCH.h
+++ b/Engine/Source/Runtime/Engine/Public/EngineSharedPCH.h
@@ -210,7 +210,6 @@
 #include "Serialization/BufferReader.h"
 #include "Containers/StaticArray.h"
 #include "Misc/SecureHash.h"
-#include "Templates/ScopedPointer.h"
 #include "Misc/Paths.h"
 #include "Modules/ModuleManager.h"
 #include "Modules/Boilerplate/ModuleBoilerplate.h"
diff --git a/Engine/Source/Runtime/Engine/Public/MaterialShared.h b/Engine/Source/Runtime/Engine/Public/MaterialShared.h
--- a/Engine/Source/Runtime/Engine/Public/MaterialShared.h
+++ b/Engine/Source/Runtime/Engine/Public/MaterialShared.h
@@ -11,7 +11,6 @@
 #include "Misc/Guid.h"
 #include "Engine/EngineTypes.h"
 #include "Templates/RefCounting.h"
-#include "Templates/ScopedPointer.h"
 #include "Templates/UniquePtr.h"
 #include "Misc/SecureHash.h"
 #include "RHI.h"
diff --git a/Engine/Source/Runtime/Engine/Public/Model.h b/Engine/Source/Runtime/Engine/Public/Model.h
--- a/Engine/Source/Runtime/Engine/Public/Model.h
+++ b/Engine/Source/Runtime/Engine/Public/Model.h
@@ -13,7 +13,6 @@
 #include "Misc/Guid.h"
 #include "Engine/EngineTypes.h"
 #include "RenderCommandFence.h"
-#include "Templates/ScopedPointer.h"
 #include "RenderResource.h"
 #include "PackedNormal.h"
 #include "Containers/DynamicRHIResourceArray.h"
diff --git a/Engine/Source/Runtime/Engine/Public/ShaderCompiler.h b/Engine/Source/Runtime/Engine/Public/ShaderCompiler.h
--- a/Engine/Source/Runtime/Engine/Public/ShaderCompiler.h
+++ b/Engine/Source/Runtime/Engine/Public/ShaderCompiler.h
@@ -8,7 +8,6 @@
 
 #include "CoreMinimal.h"
 #include "Templates/RefCounting.h"
-#include "Templates/ScopedPointer.h"
 #include "HAL/PlatformProcess.h"
 #include "ShaderCore.h"
 #include "Shader.h"
diff --git a/Engine/Source/Runtime/Engine/Public/StaticMeshResources.h b/Engine/Source/Runtime/Engine/Public/StaticMeshResources.h
--- a/Engine/Source/Runtime/Engine/Public/StaticMeshResources.h
+++ b/Engine/Source/Runtime/Engine/Public/StaticMeshResources.h
@@ -11,7 +11,6 @@
 #include "Misc/Guid.h"
 #include "Engine/EngineTypes.h"
 #include "UObject/UObjectIterator.h"
-#include "Templates/ScopedPointer.h"
 #include "Materials/MaterialInterface.h"
 #include "RenderResource.h"
 #include "PackedNormal.h"
diff --git a/Engine/Source/Runtime/Engine/Public/TextureResource.h b/Engine/Source/Runtime/Engine/Public/TextureResource.h
--- a/Engine/Source/Runtime/Engine/Public/TextureResource.h
+++ b/Engine/Source/Runtime/Engine/Public/TextureResource.h
@@ -11,7 +11,6 @@
 #include "Containers/IndirectArray.h"
 #include "Stats/Stats.h"
 #include "Containers/List.h"
-#include "Templates/ScopedPointer.h"
 #include "Async/AsyncWork.h"
 #include "Async/AsyncFileHandle.h"
 #include "RHI.h"
diff --git a/Engine/Source/Runtime/Landscape/Private/LandscapeGrass.cpp b/Engine/Source/Runtime/Landscape/Private/LandscapeGrass.cpp
--- a/Engine/Source/Runtime/Landscape/Private/LandscapeGrass.cpp
+++ b/Engine/Source/Runtime/Landscape/Private/LandscapeGrass.cpp
@@ -3,7 +3,6 @@
 #include "CoreMinimal.h"
 #include "GenericPlatform/GenericPlatformStackWalk.h"
 #include "HAL/FileManager.h"
-#include "Templates/ScopedPointer.h"
 #include "Misc/FileHelper.h"
 #include "Misc/Paths.h"
 #include "Misc/Guid.h"
diff --git a/Engine/Source/Runtime/Launch/Private/HTML5/LaunchHTML5.cpp b/Engine/Source/Runtime/Launch/Private/HTML5/LaunchHTML5.cpp
--- a/Engine/Source/Runtime/Launch/Private/HTML5/LaunchHTML5.cpp
+++ b/Engine/Source/Runtime/Launch/Private/HTML5/LaunchHTML5.cpp
@@ -4,7 +4,6 @@
 #include "Misc/App.h"
 #include "Misc/OutputDeviceError.h"
 #include "HTML5/HTML5PlatformMisc.h"
-#include "Templates/ScopedPointer.h"
 #include "LaunchEngineLoop.h"
 #include "EngineAnalytics.h"
 #include "Interfaces/IAnalyticsProvider.h"
@@ -122,7 +121,7 @@ int main(int argc, char** argv)
 {
 	// Specify the application wide locale to be UTF-8 aware. Without this,
 	// wprintf family of functions will fail on not being able to handle
-	// non-ASCII characters such as Scandinavian å, ä and ö. Even more,
+	// non-ASCII characters such as Scandinavian Ã¥, Ã¤ and Ã¶. Even more,
 	// current UE_LOG() handling is not able to cope with these kind of
 	// failures, and would crash.
 	setlocale(LC_ALL, "C.UTF-8");
diff --git a/Engine/Source/Runtime/Launch/Private/LaunchEngineLoop.cpp b/Engine/Source/Runtime/Launch/Private/LaunchEngineLoop.cpp
--- a/Engine/Source/Runtime/Launch/Private/LaunchEngineLoop.cpp
+++ b/Engine/Source/Runtime/Launch/Private/LaunchEngineLoop.cpp
@@ -24,7 +24,6 @@
 #include "Misc/App.h"
 #include "Misc/OutputDeviceConsole.h"
 #include "HAL/PlatformFilemanager.h"
-#include "Templates/ScopedPointer.h"
 #include "HAL/FileManagerGeneric.h"
 #include "HAL/ExceptionHandling.h"
 #include "Stats/StatsMallocProfilerProxy.h"
diff --git a/Engine/Source/Runtime/NetworkFile/Private/NetworkPlatformFile.cpp b/Engine/Source/Runtime/NetworkFile/Private/NetworkPlatformFile.cpp
--- a/Engine/Source/Runtime/NetworkFile/Private/NetworkPlatformFile.cpp
+++ b/Engine/Source/Runtime/NetworkFile/Private/NetworkPlatformFile.cpp
@@ -1,7 +1,6 @@
 // Copyright 1998-2019 Epic Games, Inc. All Rights Reserved.
 
 #include "NetworkPlatformFile.h"
-#include "Templates/ScopedPointer.h"
 #include "Misc/CommandLine.h"
 #include "Misc/Paths.h"
 #include "Misc/ScopedEvent.h"
diff --git a/Engine/Source/Runtime/Online/BuildPatchServices/Private/BuildPatchGeneration.cpp b/Engine/Source/Runtime/Online/BuildPatchServices/Private/BuildPatchGeneration.cpp
--- a/Engine/Source/Runtime/Online/BuildPatchServices/Private/BuildPatchGeneration.cpp
+++ b/Engine/Source/Runtime/Online/BuildPatchServices/Private/BuildPatchGeneration.cpp
@@ -1,7 +1,6 @@
 // Copyright 1998-2019 Epic Games, Inc. All Rights Reserved.
 
 #include "BuildPatchGeneration.h"
-#include "Templates/ScopedPointer.h"
 #include "Templates/UniquePtr.h"
 #include "Templates/Greater.h"
 #include "Templates/Tuple.h"
diff --git a/Engine/Source/Runtime/Online/BuildPatchServices/Private/Generation/FileAttributesParser.cpp b/Engine/Source/Runtime/Online/BuildPatchServices/Private/Generation/FileAttributesParser.cpp
--- a/Engine/Source/Runtime/Online/BuildPatchServices/Private/Generation/FileAttributesParser.cpp
+++ b/Engine/Source/Runtime/Online/BuildPatchServices/Private/Generation/FileAttributesParser.cpp
@@ -1,6 +1,5 @@
 // Copyright 1998-2019 Epic Games, Inc. All Rights Reserved.
 #include "Generation/FileAttributesParser.h"
-#include "Templates/ScopedPointer.h"
 #include "Misc/FileHelper.h"
 #include "Templates/UniquePtr.h"
 
diff --git a/Engine/Source/Runtime/PakFile/Public/IPlatformFilePak.h b/Engine/Source/Runtime/PakFile/Public/IPlatformFilePak.h
--- a/Engine/Source/Runtime/PakFile/Public/IPlatformFilePak.h
+++ b/Engine/Source/Runtime/PakFile/Public/IPlatformFilePak.h
@@ -7,7 +7,6 @@
 #include "Stats/Stats.h"
 #include "Misc/Paths.h"
 #include "Misc/ScopeLock.h"
-#include "Templates/ScopedPointer.h"
 #include "Templates/UniquePtr.h"
 #include "Math/BigInt.h"
 #include "Misc/AES.h"
diff --git a/Engine/Source/Runtime/RenderCore/Public/GlobalShader.h b/Engine/Source/Runtime/RenderCore/Public/GlobalShader.h
--- a/Engine/Source/Runtime/RenderCore/Public/GlobalShader.h
+++ b/Engine/Source/Runtime/RenderCore/Public/GlobalShader.h
@@ -7,7 +7,6 @@
 #pragma once
 
 #include "CoreMinimal.h"
-#include "Templates/ScopedPointer.h"
 #include "Shader.h"
 #include "Templates/UniquePtr.h"
 #include "RHIResources.h"
diff --git a/Engine/Source/Runtime/Renderer/Private/DistanceFieldLightingShared.h b/Engine/Source/Runtime/Renderer/Private/DistanceFieldLightingShared.h
--- a/Engine/Source/Runtime/Renderer/Private/DistanceFieldLightingShared.h
+++ b/Engine/Source/Runtime/Renderer/Private/DistanceFieldLightingShared.h
@@ -23,8 +23,6 @@ class FSceneRenderer;
 class FShaderParameterMap;
 class FViewInfo;
 
-template<typename ReferencedType> class TScopedPointer;
-
 DECLARE_LOG_CATEGORY_EXTERN(LogDistanceField, Warning, All);
 
 /** Tile sized used for most AO compute shaders. */
diff --git a/Engine/Source/Runtime/Renderer/Private/LightSceneInfo.h b/Engine/Source/Runtime/Renderer/Private/LightSceneInfo.h
--- a/Engine/Source/Runtime/Renderer/Private/LightSceneInfo.h
+++ b/Engine/Source/Runtime/Renderer/Private/LightSceneInfo.h
@@ -8,7 +8,6 @@
 
 #include "CoreMinimal.h"
 #include "RenderResource.h"
-#include "Templates/ScopedPointer.h"
 #include "Math/GenericOctreePublic.h"
 #include "SceneManagement.h"
 #include "Math/GenericOctree.h"
diff --git a/Engine/Source/Runtime/Renderer/Private/SceneRendering.h b/Engine/Source/Runtime/Renderer/Private/SceneRendering.h
--- a/Engine/Source/Runtime/Renderer/Private/SceneRendering.h
+++ b/Engine/Source/Runtime/Renderer/Private/SceneRendering.h
@@ -12,7 +12,6 @@
 #include "Stats/Stats.h"
 #include "RHI.h"
 #include "RenderResource.h"
-#include "Templates/ScopedPointer.h"
 #include "UniformBuffer.h"
 #include "GlobalDistanceFieldParameters.h"
 #include "SceneView.h"
diff --git a/Engine/Source/Runtime/SandboxFile/Public/IPlatformFileSandboxWrapper.h b/Engine/Source/Runtime/SandboxFile/Public/IPlatformFileSandboxWrapper.h
--- a/Engine/Source/Runtime/SandboxFile/Public/IPlatformFileSandboxWrapper.h
+++ b/Engine/Source/Runtime/SandboxFile/Public/IPlatformFileSandboxWrapper.h
@@ -5,7 +5,6 @@
 
 #include "CoreMinimal.h"
 #include "GenericPlatform/GenericPlatformFile.h"
-#include "Templates/ScopedPointer.h"
 #include "Misc/Paths.h"
 #include "Templates/UniquePtr.h"
 
diff --git a/Engine/Source/Runtime/Slate/Public/SlateSharedPCH.h b/Engine/Source/Runtime/Slate/Public/SlateSharedPCH.h
--- a/Engine/Source/Runtime/Slate/Public/SlateSharedPCH.h
+++ b/Engine/Source/Runtime/Slate/Public/SlateSharedPCH.h
@@ -211,7 +211,6 @@
 #include "Async/TaskGraphInterfaces.h"
 #include "Misc/CoreDelegates.h"
 #include "Misc/IQueuedWork.h"
-#include "Templates/ScopedPointer.h"
 #include "Misc/OutputDeviceError.h"
 #include "Misc/ObjectThumbnail.h"
 #include "Serialization/CustomVersion.h"
diff --git a/Engine/Source/Runtime/StreamingFile/Private/StreamingNetworkPlatformFile.cpp b/Engine/Source/Runtime/StreamingFile/Private/StreamingNetworkPlatformFile.cpp
--- a/Engine/Source/Runtime/StreamingFile/Private/StreamingNetworkPlatformFile.cpp
+++ b/Engine/Source/Runtime/StreamingFile/Private/StreamingNetworkPlatformFile.cpp
@@ -1,7 +1,6 @@
 // Copyright 1998-2019 Epic Games, Inc. All Rights Reserved.
 
 #include "StreamingNetworkPlatformFile.h"
-#include "Templates/ScopedPointer.h"
 #include "Misc/CommandLine.h"
 #include "Misc/ScopeLock.h"
 #include "Modules/ModuleManager.h"
