Trace to a file using the command line argument -tracefile=[fs_path]

diff --git a/Engine/Source/Runtime/Launch/Private/LaunchEngineLoop.cpp b/Engine/Source/Runtime/Launch/Private/LaunchEngineLoop.cpp
--- a/Engine/Source/Runtime/Launch/Private/LaunchEngineLoop.cpp
+++ b/Engine/Source/Runtime/Launch/Private/LaunchEngineLoop.cpp
@@ -1219,10 +1219,15 @@ int32 FEngineLoop::PreInit(const TCHAR* CmdLine)
 	}
 
 	{
-		FString TraceHost;
-		if (FParse::Value(CmdLine, TEXT("-tracehost="), TraceHost))
+		FString Parameter;
+		if (FParse::Value(CmdLine, TEXT("-tracehost="), Parameter))
 		{
-			Trace::Connect(*TraceHost);
+			Trace::Connect(*Parameter);
+		}
+
+		else if (FParse::Value(CmdLine, TEXT("-tracefile="), Parameter))
+		{
+			Trace::Open(*Parameter);
 		}
 
 #if PLATFORM_WINDOWS && !UE_BUILD_SHIPPING
diff --git a/Engine/Source/Runtime/TraceLog/Public/Trace/Trace.h b/Engine/Source/Runtime/TraceLog/Public/Trace/Trace.h
--- a/Engine/Source/Runtime/TraceLog/Public/Trace/Trace.h
+++ b/Engine/Source/Runtime/TraceLog/Public/Trace/Trace.h
@@ -9,6 +9,7 @@ namespace Trace
 {
 
 UE_TRACE_API bool	Connect(const TCHAR* Host) UE_TRACE_IMPL(false);
+UE_TRACE_API bool	Open(const TCHAR* Path) UE_TRACE_IMPL(false);
 UE_TRACE_API uint32 ToggleEvent(const TCHAR* Wildcard, bool bState) UE_TRACE_IMPL(false);
 UE_TRACE_API void	Flush() UE_TRACE_IMPL();
 
