Fixed being able to override save directory with argument alias file
#rnx

diff --git a/Engine/Source/Runtime/Launch/Private/LaunchEngineLoop.cpp b/Engine/Source/Runtime/Launch/Private/LaunchEngineLoop.cpp
--- a/Engine/Source/Runtime/Launch/Private/LaunchEngineLoop.cpp
+++ b/Engine/Source/Runtime/Launch/Private/LaunchEngineLoop.cpp
@@ -881,8 +881,18 @@ void LaunchCheckForCommandLineAliases(bool& bChanged)
 	bChanged = false;
 	if (FPaths::IsProjectFilePathSet())
 	{
+		// Pass GeneratedConfigDir as nullptr instead of FPaths::GeneratedConfigDir() so that saved directory is not cached before -saveddir argument can be added
+		const TCHAR* GeneratedConfigDir = nullptr;
+
 		FConfigFile Config;
-		if (FConfigCacheIni::LoadExternalIniFile(Config, TEXT("CommandLineAliases"), nullptr, *FPaths::Combine(FPaths::ProjectDir(), TEXT("Config")), false))
+		if (FConfigCacheIni::LoadExternalIniFile(Config, TEXT("CommandLineAliases"), nullptr, *FPaths::Combine(FPaths::ProjectDir(), TEXT("Config")),
+			/*bIsBaseIniName=*/ false,
+			/*Platform=*/ nullptr,
+			/*bForceReload=*/ false,
+			/*bWriteDestIni=*/ false,
+			/*bAllowGeneratedIniWhenCooked=*/ true,
+			GeneratedConfigDir
+			))
 		{
 			if (FConfigSection* Section = Config.Find(TEXT("CommandLineAliases")))
 			{
