When processing command line aliases, check for there being a valid ProjectPath rather than a project file path. In some build types, the project file path might not have been configured and therefore the aliasing system will early out.

diff --git a/Engine/Source/Runtime/Launch/Private/LaunchEngineLoop.cpp b/Engine/Source/Runtime/Launch/Private/LaunchEngineLoop.cpp
--- a/Engine/Source/Runtime/Launch/Private/LaunchEngineLoop.cpp
+++ b/Engine/Source/Runtime/Launch/Private/LaunchEngineLoop.cpp
@@ -879,7 +879,7 @@ bool LaunchCheckForFileOverride(const TCHAR* CmdLine, bool& OutFileOverrideFound
 void LaunchCheckForCommandLineAliases(bool& bChanged)
 {
 	bChanged = false;
-	if (FPaths::IsProjectFilePathSet())
+	if (FPaths::ProjectDir().Len() > 0)
 	{
 		// Pass GeneratedConfigDir as nullptr instead of FPaths::GeneratedConfigDir() so that saved directory is not cached before -saveddir argument can be added
 		const TCHAR* GeneratedConfigDir = nullptr;
@@ -962,7 +962,7 @@ void LaunchCheckForCmdLineFile(const TCHAR* CmdLine, bool& bChanged)
 			};
 
 			bool bFoundFile = TryProcessFile(CmdLineFile);
-			if (!bFoundFile && FPaths::IsProjectFilePathSet())
+			if (!bFoundFile && FPaths::ProjectDir().Len() > 0)
 			{
 				const FString ProjectDir = FPaths::ProjectDir();
 				bFoundFile = TryProcessFile(ProjectDir + CmdLineFile);
