CmdLineFile argument now removes itself from command line
#rnx

diff --git a/Engine/Source/Runtime/Launch/Private/LaunchEngineLoop.cpp b/Engine/Source/Runtime/Launch/Private/LaunchEngineLoop.cpp
--- a/Engine/Source/Runtime/Launch/Private/LaunchEngineLoop.cpp
+++ b/Engine/Source/Runtime/Launch/Private/LaunchEngineLoop.cpp
@@ -947,12 +947,26 @@ void LaunchCheckForCmdLineFile(const TCHAR* CmdLine, bool& bChanged)
 				FString FileCmds;
 				if (FFileHelper::LoadFileToString(FileCmds, *InFilePath))
 				{
-					FileCmds = FString(TEXT(" ")) + FileCmds.TrimStartAndEnd();
-					if (FileCmds.Len() > 1)
+					FileCmds = FileCmds.TrimStartAndEnd();
+					if (FileCmds.Len() > 0)
 					{
-						UE_LOG(LogInit, Log, TEXT("Appending commandline from file: %s, %s"), *InFilePath, *FileCmds);
-						FCommandLine::Append(*FileCmds);
-						bChanged = true;
+						UE_LOG(LogInit, Log, TEXT("Inserting commandline from file: %s, %s"), *InFilePath, *FileCmds);
+						FString NewCommandLine = FCommandLine::Get();
+						FString Left;
+						FString Right;
+						if (NewCommandLine.Split(TEXT("-CmdLineFile="), &Left, &Right))
+						{
+							FString NextToken;
+							const TCHAR* Stream = *Right;
+							if (FParse::Token(Stream, NextToken, /*bUseEscape=*/ false))
+							{
+								Right = FString(Stream);
+							}
+
+							NewCommandLine = Left + TEXT(" ") + FileCmds + TEXT(" ") + Right;
+							FCommandLine::Set(*NewCommandLine);
+							bChanged = true;
+						}
 					}
 
 					return true;
