Remove explict calls to InitGameTextLocalization() as they shouldn't be needed anymore.

diff --git a/Engine/Source/Runtime/Launch/Private/LaunchEngineLoop.cpp b/Engine/Source/Runtime/Launch/Private/LaunchEngineLoop.cpp
--- a/Engine/Source/Runtime/Launch/Private/LaunchEngineLoop.cpp
+++ b/Engine/Source/Runtime/Launch/Private/LaunchEngineLoop.cpp
@@ -2894,8 +2894,6 @@ int32 FEngineLoop::PreInitPostStartupScreen(const TCHAR* CmdLine)
 			}
 		}
 
-		InitGameTextLocalization();
-
 		DECLARE_SCOPE_CYCLE_COUNTER(TEXT("Initial UObject load"), STAT_InitialUObjectLoad, STATGROUP_LoadTime);
 
 		// In order to be able to use short script package names get all script
@@ -3447,7 +3445,6 @@ int32 FEngineLoop::PreInitPostStartupScreen(const TCHAR* CmdLine)
 
 #else // WITH_ENGINE
 	InitEngineTextLocalization();
-	InitGameTextLocalization();
 #if USE_LOCALIZED_PACKAGE_CACHE
 	{
 		SCOPED_BOOT_TIMING("FPackageLocalizationManager::Get().InitializeFromDefaultCache");
