Attempt to enforce the correct casing for project file path to avoid problems when packaging/running builds from the editor to target platforms that are case sensitive (eg: Android & iOS).

Additional testing needed as I lack hardware/testing environment to confirm that there aren't additional issues that need to be addressed.

diff --git a/Engine/Source/Runtime/Launch/Private/LaunchEngineLoop.cpp b/Engine/Source/Runtime/Launch/Private/LaunchEngineLoop.cpp
--- a/Engine/Source/Runtime/Launch/Private/LaunchEngineLoop.cpp
+++ b/Engine/Source/Runtime/Launch/Private/LaunchEngineLoop.cpp
@@ -648,7 +648,7 @@ void LaunchFixGameNameCase()
 	// correct the case of the game name, if possible (unless we're running a program and the game name is already set)
 	if (FPaths::IsProjectFilePathSet())
 	{
-		const FString GameName(FPaths::GetBaseFilename(IFileManager::Get().GetFilenameOnDisk(*FPaths::GetProjectFilePath())));
+		const FString GameName(FPaths::GetBaseFilename(FPaths::GetProjectFilePath()));
 
 		const bool bGameNameMatchesProjectCaseSensitive = (FCString::Strcmp(*GameName, FApp::GetProjectName()) == 0);
 		if (!bGameNameMatchesProjectCaseSensitive && (FApp::IsProjectNameEmpty() || GIsGameAgnosticExe || (GameName.Len() > 0 && GIsGameAgnosticExe)))
@@ -1739,6 +1739,12 @@ int32 FEngineLoop::PreInitPreStartupScreen(const TCHAR* CmdLine)
 	}
 
 #if !IS_PROGRAM
+	// Fix the project file path case before we attempt to fix the game name
+	if (FPaths::IsProjectFilePathSet())
+	{
+		FPaths::SetProjectFilePath(IFileManager::Get().GetFilenameOnDisk(*FPaths::GetProjectFilePath()));
+	}
+
 	if (FApp::HasProjectName())
 	{
 		// Tell the module manager what the game binaries folder is
