Don't auto-connect Programs to a trace recorder.

#rnx

diff --git a/Engine/Source/Runtime/Launch/Private/LaunchEngineLoop.cpp b/Engine/Source/Runtime/Launch/Private/LaunchEngineLoop.cpp
--- a/Engine/Source/Runtime/Launch/Private/LaunchEngineLoop.cpp
+++ b/Engine/Source/Runtime/Launch/Private/LaunchEngineLoop.cpp
@@ -1392,7 +1392,7 @@ int32 FEngineLoop::PreInitPreStartupScreen(const TCHAR* CmdLine)
 			Trace::WriteTo(*Parameter);
 		}
 
-#if PLATFORM_WINDOWS && !UE_BUILD_SHIPPING
+#if PLATFORM_WINDOWS && !UE_BUILD_SHIPPING && !IS_PROGRAM
 		else
 		{
 			// If we can detect a named event then we can try and auto-connect to UnrealInsights.
