Avoid spamming "Frame %d" events when Trace  cpu channel is enabled.

diff --git a/Engine/Source/Runtime/Launch/Private/LaunchEngineLoop.cpp b/Engine/Source/Runtime/Launch/Private/LaunchEngineLoop.cpp
--- a/Engine/Source/Runtime/Launch/Private/LaunchEngineLoop.cpp
+++ b/Engine/Source/Runtime/Launch/Private/LaunchEngineLoop.cpp
@@ -4499,11 +4499,18 @@ void FEngineLoop::Tick()
 
 	uint64 CurrentFrameCounter = GFrameCounter;
 
+	if (UE_TRACE_CHANNELEXPR_IS_ENABLED(CpuChannel))
+	{
+		TRACE_CPUPROFILER_EVENT_SCOPE(Frame);
+	}
+	else
+	{
 #if PLATFORM_LIMIT_PROFILER_UNIQUE_NAMED_EVENTS
-	SCOPED_NAMED_EVENT(FEngineLoopTick, FColor::Red);
+		SCOPED_NAMED_EVENT(FEngineLoopTick, FColor::Red);
 #else
-	SCOPED_NAMED_EVENT_F(TEXT("Frame %d"), FColor::Red, CurrentFrameCounter);
+		SCOPED_NAMED_EVENT_F(TEXT("Frame %d"), FColor::Red, CurrentFrameCounter);
 #endif
+	}
 
 	// execute callbacks for cvar changes
 	{
