Fixed check(FPlatformTLS::GetCurrentThreadId() == MasterThreadID); asserts. The GLog singleton is lazy initialised and by default will assume that its "master thread" is the one it was created on. This lazy initialisation can happen during the dynamic init (e.g. static) of a DLL which modern Windows does on a worker thread thus makeing its master thread not this one. So we make it this one and GLog->TearDown() is happy.

#rnx

diff --git a/Engine/Source/Runtime/Launch/Private/LaunchEngineLoop.cpp b/Engine/Source/Runtime/Launch/Private/LaunchEngineLoop.cpp
--- a/Engine/Source/Runtime/Launch/Private/LaunchEngineLoop.cpp
+++ b/Engine/Source/Runtime/Launch/Private/LaunchEngineLoop.cpp
@@ -1335,6 +1335,16 @@ int32 FEngineLoop::PreInitPreStartupScreen(const TCHAR* CmdLine)
 
 	SCOPED_BOOT_TIMING("FEngineLoop::PreInit");
 
+	// The GLog singleton is lazy initialised and by default will assume that
+	// its "master thread" is the one it was created on. This lazy initialisation
+	// can happen during the dynamic init (e.g. static) of a DLL which modern
+	// Windows does on a worker thread thus makeing its master thread not this one.
+	// So we make it this one and GLog->TearDown() is happy.
+	if (GLog != nullptr)
+	{
+		GLog->SetCurrentThreadAsMasterThread();
+	}
+
 	// Set the flag for whether we've build DebugGame instead of Development. The engine does not know this (whereas the launch module does) because it is always built in development.
 #if UE_BUILD_DEVELOPMENT && defined(UE_BUILD_DEVELOPMENT_WITH_DEBUGGAME) && UE_BUILD_DEVELOPMENT_WITH_DEBUGGAME
 	FApp::SetDebugGame(true);
