Do not initialize HighResScreenshotConfig on mobile platforms, to avoid creating 6 threads

diff --git a/Engine/Source/Runtime/Launch/Private/LaunchEngineLoop.cpp b/Engine/Source/Runtime/Launch/Private/LaunchEngineLoop.cpp
--- a/Engine/Source/Runtime/Launch/Private/LaunchEngineLoop.cpp
+++ b/Engine/Source/Runtime/Launch/Private/LaunchEngineLoop.cpp
@@ -3516,11 +3516,13 @@ int32 FEngineLoop::PreInitPostStartupScreen(const TCHAR* CmdLine)
 	}
 #endif
 
+#if !(PLATFORM_IOS || PLATFORM_TVOS || PLATFORM_ANDROID)
 	// Init HighRes screenshot system, unless running on server
 	if (!IsRunningDedicatedServer())
 	{
 		GetHighResScreenshotConfig().Init();
 	}
+#endif
 
 #else // WITH_ENGINE
 	InitEngineTextLocalization();
