Fix detection of engine modules being out of date.

diff --git a/Engine/Source/Developer/DesktopPlatform/Private/DesktopPlatformBase.cpp b/Engine/Source/Developer/DesktopPlatform/Private/DesktopPlatformBase.cpp
--- a/Engine/Source/Developer/DesktopPlatform/Private/DesktopPlatformBase.cpp
+++ b/Engine/Source/Developer/DesktopPlatform/Private/DesktopPlatformBase.cpp
@@ -558,7 +558,7 @@ bool FDesktopPlatformBase::CleanGameProject(const FString& ProjectDir, FString&
 	return true;
 }
 
-bool FDesktopPlatformBase::CompileGameProject(const FString& RootDir, const FString& ProjectFileName, FFeedbackContext* Warn)
+bool FDesktopPlatformBase::CompileGameProject(const FString& RootDir, const FString& ProjectFileName, FFeedbackContext* Warn, ECompilationResult::Type* OutResult)
 {
 	FModuleManager& ModuleManager = FModuleManager::Get();
 
@@ -579,7 +579,12 @@ bool FDesktopPlatformBase::CompileGameProject(const FString& RootDir, const FStr
 	Arguments += " -Progress -NoEngineChanges -NoHotReloadFromIDE";
 
 	// Run UBT
-	bool bResult = RunUnrealBuildTool(LOCTEXT("CompilingProject", "Compiling project..."), RootDir, Arguments, Warn);
+	int32 ExitCode;
+	bool bResult = static_cast<IDesktopPlatform*>(this)->RunUnrealBuildTool(LOCTEXT("CompilingProject", "Compiling project..."), RootDir, Arguments, Warn, ExitCode);
+	if (OutResult != nullptr)
+	{
+		*OutResult = (ECompilationResult::Type)ExitCode;
+	}
 
 	// Reset module paths in case they have changed during compilation
 	ModuleManager.ResetModulePathsCache();
@@ -749,6 +754,12 @@ FProcHandle FDesktopPlatformBase::InvokeUnrealBuildToolAsync(const FString& InCm
 	return ProcHandle;
 }
 
+bool FDesktopPlatformBase::RunUnrealBuildTool(const FText& Description, const FString& RootDir, const FString& Arguments, FFeedbackContext* Warn)
+{
+	int32 ExitCode;
+	return static_cast<IDesktopPlatform*>(this)->RunUnrealBuildTool(Description, RootDir, Arguments, Warn, ExitCode);
+}
+
 struct FTargetFileVisitor : IPlatformFile::FDirectoryStatVisitor
 {
 	TSet<FString>& RemainingTargetNames;
diff --git a/Engine/Source/Developer/DesktopPlatform/Private/DesktopPlatformBase.h b/Engine/Source/Developer/DesktopPlatform/Private/DesktopPlatformBase.h
--- a/Engine/Source/Developer/DesktopPlatform/Private/DesktopPlatformBase.h
+++ b/Engine/Source/Developer/DesktopPlatform/Private/DesktopPlatformBase.h
@@ -6,6 +6,7 @@
 #include "HAL/PlatformProcess.h"
 #include "IDesktopPlatform.h"
 #include "Misc/UProjectInfo.h"
+#include "Misc/CompilationResult.h"
 
 class FEngineVersion;
 class FJsonObject;
@@ -45,11 +46,12 @@ public:
 	virtual bool OpenProject(const FString& ProjectFileName) override;
 
 	virtual bool CleanGameProject(const FString& ProjectDir, FString& OutFailPath, FFeedbackContext* Warn) override;
-	virtual bool CompileGameProject(const FString& RootDir, const FString& ProjectFileName, FFeedbackContext* Warn) override;
+	virtual bool CompileGameProject(const FString& RootDir, const FString& ProjectFileName, FFeedbackContext* Warn, ECompilationResult::Type* OutResult) override;
 	virtual bool GenerateProjectFiles(const FString& RootDir, const FString& ProjectFileName, FFeedbackContext* Warn, FString LogFilePath = FString()) override;
 	virtual bool IsUnrealBuildToolAvailable() override;
 	virtual bool InvokeUnrealBuildToolSync(const FString& InCmdLineParams, FOutputDevice &Ar, bool bSkipBuildUBT, int32& OutReturnCode, FString& OutProcOutput) override;
 	virtual FProcHandle InvokeUnrealBuildToolAsync(const FString& InCmdLineParams, FOutputDevice &Ar, void*& OutReadPipe, void*& OutWritePipe, bool bSkipBuildUBT = false) override;
+	virtual bool RunUnrealBuildTool(const FText& Description, const FString& RootDir, const FString& Arguments, FFeedbackContext* Warn) override final;
 
 	virtual const TArray<FTargetInfo>& GetTargetsForProject(const FString& ProjectFile) const override;
 	virtual const TArray<FTargetInfo>& GetTargetsForCurrentProject() const override;
diff --git a/Engine/Source/Developer/DesktopPlatform/Private/Linux/DesktopPlatformLinux.cpp b/Engine/Source/Developer/DesktopPlatform/Private/Linux/DesktopPlatformLinux.cpp
--- a/Engine/Source/Developer/DesktopPlatform/Private/Linux/DesktopPlatformLinux.cpp
+++ b/Engine/Source/Developer/DesktopPlatform/Private/Linux/DesktopPlatformLinux.cpp
@@ -414,8 +414,10 @@ bool FDesktopPlatformLinux::OpenProject(const FString &ProjectFileName)
 	return false;
 }
 
-bool FDesktopPlatformLinux::RunUnrealBuildTool(const FText& Description, const FString& RootDir, const FString& Arguments, FFeedbackContext* Warn)
+bool FDesktopPlatformLinux::RunUnrealBuildTool(const FText& Description, const FString& RootDir, const FString& Arguments, FFeedbackContext* Warn, int32& OutExitCode)
 {
+	OutExitCode = 1;
+
 	// Get the path to UBT
 	FString UnrealBuildToolPath = RootDir / TEXT("Engine/Binaries/DotNET/UnrealBuildTool.exe");
 	if(IFileManager::Get().FileSize(*UnrealBuildToolPath) < 0)
@@ -432,8 +434,7 @@ bool FDesktopPlatformLinux::RunUnrealBuildTool(const FText& Description, const F
 	FString CmdLineParams = FString::Printf(TEXT("\"%s\" \"%s\" %s"), *ScriptPath, *UnrealBuildToolPath, *Arguments);
 
 	// Spawn it with bash (and not sh) because of pushd
-	int32 ExitCode = 0;
-	return FFeedbackContextMarkup::PipeProcessOutput(Description, TEXT("/bin/bash"), CmdLineParams, Warn, &ExitCode) && ExitCode == 0;
+	return FFeedbackContextMarkup::PipeProcessOutput(Description, TEXT("/bin/bash"), CmdLineParams, Warn, &OutExitCode) && OutExitCode == 0;
 }
 
 bool FDesktopPlatformLinux::IsUnrealBuildToolRunning()
diff --git a/Engine/Source/Developer/DesktopPlatform/Private/Linux/DesktopPlatformLinux.h b/Engine/Source/Developer/DesktopPlatform/Private/Linux/DesktopPlatformLinux.h
--- a/Engine/Source/Developer/DesktopPlatform/Private/Linux/DesktopPlatformLinux.h
+++ b/Engine/Source/Developer/DesktopPlatform/Private/Linux/DesktopPlatformLinux.h
@@ -25,7 +25,7 @@ public:
 
 	virtual bool OpenProject(const FString &ProjectFileName) override;
 
-	virtual bool RunUnrealBuildTool(const FText& Description, const FString& RootDir, const FString& Arguments, FFeedbackContext* Warn) override;
+	virtual bool RunUnrealBuildTool(const FText& Description, const FString& RootDir, const FString& Arguments, FFeedbackContext* Warn, int32& OutExitCode) override;
 	virtual bool IsUnrealBuildToolRunning() override;
 
 	virtual FFeedbackContext* GetNativeFeedbackContext() override;
diff --git a/Engine/Source/Developer/DesktopPlatform/Private/Mac/DesktopPlatformMac.cpp b/Engine/Source/Developer/DesktopPlatform/Private/Mac/DesktopPlatformMac.cpp
--- a/Engine/Source/Developer/DesktopPlatform/Private/Mac/DesktopPlatformMac.cpp
+++ b/Engine/Source/Developer/DesktopPlatform/Private/Mac/DesktopPlatformMac.cpp
@@ -597,8 +597,10 @@ bool FDesktopPlatformMac::UpdateFileAssociations()
 	return Status == noErr;
 }
 
-bool FDesktopPlatformMac::RunUnrealBuildTool(const FText& Description, const FString& RootDir, const FString& Arguments, FFeedbackContext* Warn)
+bool FDesktopPlatformMac::RunUnrealBuildTool(const FText& Description, const FString& RootDir, const FString& Arguments, FFeedbackContext* Warn, int32& OutExitCode)
 {
+	OutExitCode = 1;
+
 	// Get the path to UBT
 	FString UnrealBuildToolPath = RootDir / TEXT("Engine/Binaries/DotNET/UnrealBuildTool.exe");
 	if(IFileManager::Get().FileSize(*UnrealBuildToolPath) < 0)
@@ -612,8 +614,7 @@ bool FDesktopPlatformMac::RunUnrealBuildTool(const FText& Description, const FSt
 	FString CmdLineParams = FString::Printf(TEXT("\"%s\" \"%s\" %s"), *ScriptPath, *UnrealBuildToolPath, *Arguments);
 
 	// Spawn it
-	int32 ExitCode = 0;
-	return FFeedbackContextMarkup::PipeProcessOutput(Description, TEXT("/bin/sh"), CmdLineParams, Warn, &ExitCode) && ExitCode == 0;
+	return FFeedbackContextMarkup::PipeProcessOutput(Description, TEXT("/bin/sh"), CmdLineParams, Warn, &OutExitCode) && OutExitCode == 0;
 }
 
 bool FDesktopPlatformMac::IsUnrealBuildToolRunning()
diff --git a/Engine/Source/Developer/DesktopPlatform/Private/Mac/DesktopPlatformMac.h b/Engine/Source/Developer/DesktopPlatform/Private/Mac/DesktopPlatformMac.h
--- a/Engine/Source/Developer/DesktopPlatform/Private/Mac/DesktopPlatformMac.h
+++ b/Engine/Source/Developer/DesktopPlatform/Private/Mac/DesktopPlatformMac.h
@@ -20,7 +20,7 @@ public:
 	virtual bool VerifyFileAssociations() override;
 	virtual bool UpdateFileAssociations() override;
 
-	virtual bool RunUnrealBuildTool(const FText& Description, const FString& RootDir, const FString& Arguments, FFeedbackContext* Warn) override;
+	virtual bool RunUnrealBuildTool(const FText& Description, const FString& RootDir, const FString& Arguments, FFeedbackContext* Warn, int32& OutExitCode) override;
 	virtual bool IsUnrealBuildToolRunning() override;
 
 	virtual FFeedbackContext* GetNativeFeedbackContext() override;
diff --git a/Engine/Source/Developer/DesktopPlatform/Private/Windows/DesktopPlatformWindows.cpp b/Engine/Source/Developer/DesktopPlatform/Private/Windows/DesktopPlatformWindows.cpp
--- a/Engine/Source/Developer/DesktopPlatform/Private/Windows/DesktopPlatformWindows.cpp
+++ b/Engine/Source/Developer/DesktopPlatform/Private/Windows/DesktopPlatformWindows.cpp
@@ -458,8 +458,10 @@ bool FDesktopPlatformWindows::OpenProject(const FString &ProjectFileName)
 	return ::ShellExecuteExW(&Info) != 0;
 }
 
-bool FDesktopPlatformWindows::RunUnrealBuildTool(const FText& Description, const FString& RootDir, const FString& Arguments, FFeedbackContext* Warn)
+bool FDesktopPlatformWindows::RunUnrealBuildTool(const FText& Description, const FString& RootDir, const FString& Arguments, FFeedbackContext* Warn, int32& OutExitCode)
 {
+	OutExitCode = 1;
+
 	// Get the path to UBT
 	FString UnrealBuildToolPath = RootDir / TEXT("Engine/Binaries/DotNET/UnrealBuildTool.exe");
 	if(IFileManager::Get().FileSize(*UnrealBuildToolPath) < 0)
@@ -472,8 +474,7 @@ bool FDesktopPlatformWindows::RunUnrealBuildTool(const FText& Description, const
 	Warn->Logf(TEXT("Running %s %s"), *UnrealBuildToolPath, *Arguments);
 
 	// Spawn UBT
-	int32 ExitCode = 0;
-	return FFeedbackContextMarkup::PipeProcessOutput(Description, UnrealBuildToolPath, Arguments, Warn, &ExitCode) && ExitCode == 0;
+	return FFeedbackContextMarkup::PipeProcessOutput(Description, UnrealBuildToolPath, Arguments, Warn, &OutExitCode) && OutExitCode == 0;
 }
 
 bool FDesktopPlatformWindows::IsUnrealBuildToolRunning()
diff --git a/Engine/Source/Developer/DesktopPlatform/Private/Windows/DesktopPlatformWindows.h b/Engine/Source/Developer/DesktopPlatform/Private/Windows/DesktopPlatformWindows.h
--- a/Engine/Source/Developer/DesktopPlatform/Private/Windows/DesktopPlatformWindows.h
+++ b/Engine/Source/Developer/DesktopPlatform/Private/Windows/DesktopPlatformWindows.h
@@ -28,7 +28,7 @@ public:
 
 	virtual bool OpenProject(const FString &ProjectFileName) override;
 
-	virtual bool RunUnrealBuildTool(const FText& Description, const FString& RootDir, const FString& Arguments, FFeedbackContext* Warn) override;
+	virtual bool RunUnrealBuildTool(const FText& Description, const FString& RootDir, const FString& Arguments, FFeedbackContext* Warn, int32& OutExitCode) override;
 	virtual bool IsUnrealBuildToolRunning() override;
 
 	virtual FFeedbackContext* GetNativeFeedbackContext() override;
diff --git a/Engine/Source/Developer/DesktopPlatform/Public/IDesktopPlatform.h b/Engine/Source/Developer/DesktopPlatform/Public/IDesktopPlatform.h
--- a/Engine/Source/Developer/DesktopPlatform/Public/IDesktopPlatform.h
+++ b/Engine/Source/Developer/DesktopPlatform/Public/IDesktopPlatform.h
@@ -5,6 +5,7 @@
 #include "CoreMinimal.h"
 #include "Misc/EnumClassFlags.h"
 #include "Misc/CoreDelegates.h"
+#include "Misc/CompilationResult.h"
 #include "HAL/PlatformProcess.h"
 
 class FEngineVersion;
@@ -346,9 +347,10 @@ public:
 	* @param RootDir			Engine root directory for the project to use.
 	* @param ProjectFileName	Filename of the project to update
 	* @param Warn				Feedback context to use for progress updates
+	* @param OutResult			The compilation result. May be null.
 	* @return true if project files were generated successfully.
 	*/
-	virtual bool CompileGameProject(const FString& RootDir, const FString& ProjectFileName, FFeedbackContext* Warn) = 0;
+	virtual bool CompileGameProject(const FString& RootDir, const FString& ProjectFileName, FFeedbackContext* Warn, ECompilationResult::Type* OutResult = nullptr) = 0;
 
 	/**
 	* Generates project files for the given project.
@@ -392,6 +394,18 @@ public:
 	*/
 	virtual bool RunUnrealBuildTool(const FText& Description, const FString& RootDir, const FString& Arguments, FFeedbackContext* Warn) = 0;
 
+	/**
+	* Runs UnrealBuildTool with the given arguments.
+	*
+	* @param Description		Task description for FFeedbackContext
+	* @param RootDir			Engine root directory for the project to use.
+	* @param Arguments			Parameters for UnrealBuildTool
+	* @param Warn				Feedback context to use for progress updates
+	* @param OutExitCode		Receives the process exit code
+	* @return true if the task completed successfully.
+	*/
+	virtual bool RunUnrealBuildTool(const FText& Description, const FString& RootDir, const FString& Arguments, FFeedbackContext* Warn, int32& OutExitCode) = 0;
+
 	/**
 	* Checks if an instance of UnrealBuildTool is running.
 	*
diff --git a/Engine/Source/Programs/UnrealBuildTool/Modes/BuildMode.cs b/Engine/Source/Programs/UnrealBuildTool/Modes/BuildMode.cs
--- a/Engine/Source/Programs/UnrealBuildTool/Modes/BuildMode.cs
+++ b/Engine/Source/Programs/UnrealBuildTool/Modes/BuildMode.cs
@@ -303,7 +303,8 @@ namespace UnrealBuildTool
 							Result.AppendFormat("\n{0}", EngineChange.FullName);
 						}
 						Result.Append("\n\nPlease rebuild from an IDE instead.");
-						throw new BuildException(Result.ToString());
+						Log.TraceError("{0}", Result.ToString());
+						throw new CompilationResultException(CompilationResult.FailedDueToEngineChange);
 					}
 				}
 
diff --git a/Engine/Source/Programs/UnrealBuildTool/System/ExternalExecution.cs b/Engine/Source/Programs/UnrealBuildTool/System/ExternalExecution.cs
--- a/Engine/Source/Programs/UnrealBuildTool/System/ExternalExecution.cs
+++ b/Engine/Source/Programs/UnrealBuildTool/System/ExternalExecution.cs
@@ -45,10 +45,15 @@ namespace UnrealBuildTool
 		/// </summary>
 		FailedDueToHeaderChange = 4,
 
+		/// <summary>
+		/// Compilation failed due to the engine modules needing to be rebuilt
+		/// </summary>
+		FailedDueToEngineChange = 5,
+
 		/// <summary>
 		/// Compilation failed due to compilation errors
 		/// </summary>
-		OtherCompilationError = 5,
+		OtherCompilationError = 6,
 
 		/// <summary>
 		/// Compilation is not supported in the current build
diff --git a/Engine/Source/Runtime/Core/Public/Misc/CompilationResult.h b/Engine/Source/Runtime/Core/Public/Misc/CompilationResult.h
--- a/Engine/Source/Runtime/Core/Public/Misc/CompilationResult.h
+++ b/Engine/Source/Runtime/Core/Public/Misc/CompilationResult.h
@@ -25,8 +25,10 @@ namespace ECompilationResult
 		CrashOrAssert = 3,
 		/** Compilation failed because generated code changed which was not supported */
 		FailedDueToHeaderChange = 4,
+		/** Compilation failed due to the engine modules needing to be rebuilt */
+		FailedDueToEngineChange = 5,
 		/** Compilation failed due to compilation errors */
-		OtherCompilationError = 5,
+		OtherCompilationError = 6,
 		/** Compilation is not supported in the current build */
 		Unsupported,
 		/** Unknown error */
diff --git a/Engine/Source/Runtime/Launch/Private/LaunchEngineLoop.cpp b/Engine/Source/Runtime/Launch/Private/LaunchEngineLoop.cpp
--- a/Engine/Source/Runtime/Launch/Private/LaunchEngineLoop.cpp
+++ b/Engine/Source/Runtime/Launch/Private/LaunchEngineLoop.cpp
@@ -5397,7 +5397,9 @@ bool FEngineLoop::AppInit( )
 			// Check if any of the project or plugin modules are out of date, and the user wants to compile them.
 			TArray<FString> IncompatibleFiles;
 			ProjectManager.CheckModuleCompatibility(IncompatibleFiles);
-			PluginManager.CheckModuleCompatibility(IncompatibleFiles);
+
+			TArray<FString> IncompatibleEngineFiles;
+			PluginManager.CheckModuleCompatibility(IncompatibleFiles, IncompatibleEngineFiles);
 
 			if (IncompatibleFiles.Num() > 0)
 			{
@@ -5427,14 +5429,11 @@ bool FEngineLoop::AppInit( )
 				}
 
 				// If there are any engine modules that need building, force the user to build through the IDE
-				for (const FString& IncompatibleFile : IncompatibleFiles)
+				if(IncompatibleEngineFiles.Num() > 0)
 				{
-					if (FPaths::IsUnderDirectory(IncompatibleFile, FPaths::EngineDir()))
-					{
-						FString CompileForbidden = ModulesList + TEXT("\nEngine modules cannot be compiled at runtime. Please build through your IDE.");
-						FPlatformMisc::MessageBoxExt(EAppMsgType::Ok, *CompileForbidden, TEXT("Missing Modules"));
-						return false;
-					}
+					FString CompileForbidden = ModulesList + TEXT("\nEngine modules cannot be compiled at runtime. Please build through your IDE.");
+					FPlatformMisc::MessageBoxExt(EAppMsgType::Ok, *CompileForbidden, TEXT("Missing Modules"));
+					return false;
 				}
 
 				// Ask whether to compile before continuing
@@ -5463,7 +5462,8 @@ bool FEngineLoop::AppInit( )
 			// Try to compile it
 			FFeedbackContext *Context = (FFeedbackContext*)FDesktopPlatformModule::Get()->GetNativeFeedbackContext();
 			Context->BeginSlowTask(FText::FromString(TEXT("Starting build...")), true, true);
-			bool bCompileResult = FDesktopPlatformModule::Get()->CompileGameProject(FPaths::RootDir(), FPaths::GetProjectFilePath(), Context);
+			ECompilationResult::Type CompilationResult = ECompilationResult::Unknown;
+			bool bCompileResult = FDesktopPlatformModule::Get()->CompileGameProject(FPaths::RootDir(), FPaths::GetProjectFilePath(), Context, &CompilationResult);
 			Context->EndSlowTask();
 
 			// Check if we're running the wrong executable now
@@ -5472,10 +5472,19 @@ bool FEngineLoop::AppInit( )
 				return false;
 			}
 
+			// Check if we needed to modify engine files
+			if (!bCompileResult && CompilationResult == ECompilationResult::FailedDueToEngineChange)
+			{
+				FPlatformMisc::MessageBoxExt(EAppMsgType::Ok, TEXT("Engine modules are out of date, and cannot be compiled while the engine is running. Please build through your IDE."), TEXT("Missing Modules"));
+				return false;
+			}
+
 			// Get a list of modules which are still incompatible
 			TArray<FString> StillIncompatibleFiles;
 			ProjectManager.CheckModuleCompatibility(StillIncompatibleFiles);
-			PluginManager.CheckModuleCompatibility(StillIncompatibleFiles);
+
+			TArray<FString> StillIncompatibleEngineFiles;
+			PluginManager.CheckModuleCompatibility(StillIncompatibleFiles, StillIncompatibleEngineFiles);
 
 			if(!bCompileResult || StillIncompatibleFiles.Num() > 0)
 			{
diff --git a/Engine/Source/Runtime/Projects/Private/PluginManager.cpp b/Engine/Source/Runtime/Projects/Private/PluginManager.cpp
--- a/Engine/Source/Runtime/Projects/Private/PluginManager.cpp
+++ b/Engine/Source/Runtime/Projects/Private/PluginManager.cpp
@@ -1309,7 +1309,7 @@ bool FPluginManager::AreRequiredPluginsAvailable()
 }
 
 #if !IS_MONOLITHIC
-bool FPluginManager::CheckModuleCompatibility(TArray<FString>& OutIncompatibleModules)
+bool FPluginManager::CheckModuleCompatibility(TArray<FString>& OutIncompatibleModules, TArray<FString>& OutIncompatibleEngineModules)
 {
 	if(!ConfigureEnabledPlugins())
 	{
@@ -1320,8 +1320,15 @@ bool FPluginManager::CheckModuleCompatibility(TArray<FString>& OutIncompatibleMo
 	for(const TPair<FString, TSharedRef<FPlugin>>& PluginPair : AllPlugins)
 	{
 		const TSharedRef< FPlugin > &Plugin = PluginPair.Value;
-		if (Plugin->bEnabled && !FModuleDescriptor::CheckModuleCompatibility(Plugin->Descriptor.Modules, OutIncompatibleModules))
+
+		TArray<FString> IncompatibleModules;
+		if (Plugin->bEnabled && !FModuleDescriptor::CheckModuleCompatibility(Plugin->Descriptor.Modules, IncompatibleModules))
 		{
+			OutIncompatibleModules.Append(IncompatibleModules);
+			if (Plugin->GetLoadedFrom() == EPluginLoadedFrom::Engine)
+			{
+				OutIncompatibleEngineModules.Append(IncompatibleModules);
+			}
 			bResult = false;
 		}
 	}
diff --git a/Engine/Source/Runtime/Projects/Private/PluginManager.h b/Engine/Source/Runtime/Projects/Private/PluginManager.h
--- a/Engine/Source/Runtime/Projects/Private/PluginManager.h
+++ b/Engine/Source/Runtime/Projects/Private/PluginManager.h
@@ -102,7 +102,7 @@ public:
 	virtual void SetRegisterMountPointDelegate( const FRegisterMountPointDelegate& Delegate ) override;
 	virtual bool AreRequiredPluginsAvailable() override;
 #if !IS_MONOLITHIC
-	virtual bool CheckModuleCompatibility( TArray<FString>& OutIncompatibleModules ) override;
+	virtual bool CheckModuleCompatibility(TArray<FString>& OutIncompatibleModules, TArray<FString>& OutIncompatibleEngineModules) override;
 #endif
 	virtual TSharedPtr<IPlugin> FindPlugin(const FString& Name) override;
 	virtual TArray<TSharedRef<IPlugin>> GetEnabledPlugins() override;
diff --git a/Engine/Source/Runtime/Projects/Public/Interfaces/IPluginManager.h b/Engine/Source/Runtime/Projects/Public/Interfaces/IPluginManager.h
--- a/Engine/Source/Runtime/Projects/Public/Interfaces/IPluginManager.h
+++ b/Engine/Source/Runtime/Projects/Public/Interfaces/IPluginManager.h
@@ -223,10 +223,11 @@ public:
 	/** 
 	 * Checks whether modules for the enabled plug-ins are up to date.
 	 *
-	 * @param OutIncompatibleNames	Array to receive a list of incompatible module names.
+	 * @param OutIncompatibleModules Array to receive a list of incompatible module names.
+	 * @param OutIncompatibleEngineModules Array to receive a list of incompatible engine module names.
 	 * @returns true if the enabled plug-in modules are up to date.
 	 */
-	virtual bool CheckModuleCompatibility( TArray<FString>& OutIncompatibleModules ) = 0;
+	virtual bool CheckModuleCompatibility( TArray<FString>& OutIncompatibleModules, TArray<FString>& OutIncompatibleEngineModules ) = 0;
 #endif
 
 	/**
