Ensure that when getting correct case for project path during engine startup on the Windows platform, we don't also resolve directory junctions at the same time as doing so can cause problems for p4 operations.  This is done using the slow 'FindFirstFile' approach on EVERY SEGMENT of the project file path as was formerly done in the GetFilenameOnDisk implementation for Windows.

diff --git a/Engine/Source/Runtime/Launch/Private/LaunchEngineLoop.cpp b/Engine/Source/Runtime/Launch/Private/LaunchEngineLoop.cpp
--- a/Engine/Source/Runtime/Launch/Private/LaunchEngineLoop.cpp
+++ b/Engine/Source/Runtime/Launch/Private/LaunchEngineLoop.cpp
@@ -640,6 +640,51 @@ bool LaunchSetGameName(const TCHAR *InCmdLine, FString& OutGameProjectFilePathUn
 	return true;
 }
 
+void LaunchFixProjectPathCase()
+{
+	if (FPaths::IsProjectFilePathSet())
+	{
+#if PLATFORM_WINDOWS
+		// GetFilenameOnDisk on Windows will resolve directory junctions and resolving those here has negative consequences
+		// for workflows that use a junction at their root (eg: p4 gets confused about paths and operations fail).
+		// There is a way to get a case-accurate path on Windows without resolving directory junctions, but it is slow.
+		// We can use it here for this one-off situation without causing all uses of GetFilenameOnDisk to be slower.
+		FString Result;
+		WIN32_FIND_DATAW Data;
+		FString NormalizedFilename = FPaths::GetProjectFilePath();
+		FPaths::NormalizeFilename(NormalizedFilename);
+		while (NormalizedFilename.Len())
+		{
+			HANDLE Handle = FindFirstFileW(*NormalizedFilename, &Data);
+			if (Handle != INVALID_HANDLE_VALUE)
+			{
+				if (Result.Len())
+				{
+					Result = FString(Data.cFileName) / Result;
+				}
+				else
+				{
+					Result = Data.cFileName;
+				}
+				FindClose(Handle);
+			}
+			int32 SeparatorIndex = INDEX_NONE;
+			if (NormalizedFilename.FindLastChar('/', SeparatorIndex))
+			{
+				NormalizedFilename = NormalizedFilename.Mid(0, SeparatorIndex);
+			}
+			if (NormalizedFilename.Len() && (SeparatorIndex == INDEX_NONE || NormalizedFilename.EndsWith(TEXT(":"))))
+			{
+				Result = NormalizedFilename / Result;
+				NormalizedFilename.Empty();
+			}
+		}
+		FPaths::SetProjectFilePath(Result);
+#else
+		FPaths::SetProjectFilePath(IFileManager::Get().GetFilenameOnDisk(*FPaths::GetProjectFilePath()));
+#endif
+	}
+}
 
 void LaunchFixGameNameCase()
 {
@@ -1976,10 +2021,7 @@ int32 FEngineLoop::PreInitPreStartupScreen(const TCHAR* CmdLine)
 
 #if !IS_PROGRAM
 	// Fix the project file path case before we attempt to fix the game name
-	if (FPaths::IsProjectFilePathSet())
-	{
-		FPaths::SetProjectFilePath(IFileManager::Get().GetFilenameOnDisk(*FPaths::GetProjectFilePath()));
-	}
+	LaunchFixProjectPathCase();
 
 	if (FApp::HasProjectName())
 	{
