Implementation #2 - Ensure that when getting correct case for project path during engine startup on the Windows platform, we don't also resolve directory junctions at the same time as doing so can cause problems for p4 operations.  This is done using the slow 'FindFirstFile' approach on EVERY SEGMENT of the project file path as was formerly done in the GetFilenameOnDisk implementation for Windows.

diff --git a/Engine/Source/Runtime/Launch/Private/LaunchEngineLoop.cpp b/Engine/Source/Runtime/Launch/Private/LaunchEngineLoop.cpp
--- a/Engine/Source/Runtime/Launch/Private/LaunchEngineLoop.cpp
+++ b/Engine/Source/Runtime/Launch/Private/LaunchEngineLoop.cpp
@@ -16,6 +16,7 @@
 #include "Async/TaskGraphInterfaces.h"
 #include "Misc/TimeGuard.h"
 #include "Misc/Paths.h"
+#include "Misc/PathViews.h"
 #include "Misc/ConfigCacheIni.h"
 #include "Misc/OutputDeviceHelper.h"
 #include "Misc/OutputDeviceRedirector.h"
@@ -39,6 +40,7 @@
 #endif
 #include "HAL/ThreadManager.h"
 #include "ProfilingDebugging/ExternalProfiler.h"
+#include "Containers/StringView.h"
 #include "Containers/Ticker.h"
 
 #include "Interfaces/IPluginManager.h"
@@ -640,6 +642,49 @@ bool LaunchSetGameName(const TCHAR *InCmdLine, FString& OutGameProjectFilePathUn
 	return true;
 }
 
+void LaunchFixProjectPathCase()
+{
+	if (FPaths::IsProjectFilePathSet())
+	{
+#if PLATFORM_WINDOWS
+		// GetFilenameOnDisk on Windows will resolve directory junctions and resolving those here has negative consequences
+		// for workflows that use a junction at their root (eg: p4 gets confused about paths and operations fail).
+		// There is a way to get a case-accurate path on Windows without resolving directory junctions, but it is slow.
+		// We can use it here for this one-off situation without causing all uses of GetFilenameOnDisk to be slower.
+		WIN32_FIND_DATAW Data;
+		FString ProjectFilePath = FPaths::GetProjectFilePath();
+		for (FStringView PendingFix = ProjectFilePath; PendingFix.Len() > 0;)
+		{
+			const FStringView CurrentPathComponent = FPathViews::GetCleanFilename(PendingFix);
+
+			// Skip over all segments that are either empty or contain relative transforms or start with the volume separator, they should remain as-is
+			const bool bIsIgnoredSegment = CurrentPathComponent.IsEmpty() || CurrentPathComponent.Equals(TEXT("."_SV)) || CurrentPathComponent.Equals(TEXT(".."_SV)) || CurrentPathComponent.EndsWith(TEXT(':'));
+			if (!bIsIgnoredSegment)
+			{
+				// Temporarily null-terminate the current path component for the system call
+				TCHAR Terminator = TEXT('\0');
+				Swap(ProjectFilePath.GetCharArray()[PendingFix.Len()], Terminator);
+
+				HANDLE Handle = FindFirstFileW(*ProjectFilePath, &Data);
+				if (Handle != INVALID_HANDLE_VALUE)
+				{
+					check(CurrentPathComponent.Equals(Data.cFileName, ESearchCase::IgnoreCase));
+					FCString::Strncpy(GetData(ProjectFilePath) + PendingFix.Len() - CurrentPathComponent.Len(), Data.cFileName, CurrentPathComponent.Len() + 1);
+					FindClose(Handle);
+				}
+
+				Swap(ProjectFilePath.GetCharArray()[PendingFix.Len()], Terminator);
+			}
+
+			PendingFix.LeftChopInline(CurrentPathComponent.Len() + 1);
+		}
+		FPaths::SetProjectFilePath(ProjectFilePath);
+#else
+		FPaths::SetProjectFilePath(IFileManager::Get().GetFilenameOnDisk(*FPaths::GetProjectFilePath()));
+#endif
+	}
+}
+
 void LaunchFixGameNameCase()
 {
 #if PLATFORM_DESKTOP && !IS_PROGRAM
@@ -1975,10 +2020,7 @@ int32 FEngineLoop::PreInitPreStartupScreen(const TCHAR* CmdLine)
 
 #if !IS_PROGRAM
 	// Fix the project file path case before we attempt to fix the game name
-	if (FPaths::IsProjectFilePathSet())
-	{
-		FPaths::SetProjectFilePath(IFileManager::Get().GetFilenameOnDisk(*FPaths::GetProjectFilePath()));
-	}
+	LaunchFixProjectPathCase();
 
 	if (FApp::HasProjectName())
 	{
