Add -csvExecCmds command line option, to run exec commands on particular CSV frames. E.g:
-csvexeccmds="500:listtextures,1000:listtextures". This can be useful when doing deterministic benchmarking

diff --git a/Engine/Source/Runtime/Launch/Private/LaunchEngineLoop.cpp b/Engine/Source/Runtime/Launch/Private/LaunchEngineLoop.cpp
--- a/Engine/Source/Runtime/Launch/Private/LaunchEngineLoop.cpp
+++ b/Engine/Source/Runtime/Launch/Private/LaunchEngineLoop.cpp
@@ -1290,6 +1290,47 @@ static void UpdateCoreCsvStats_BeginFrame()
 		CSV_CUSTOM_STAT_GLOBAL(CPUUsage_Idle, IdleUsageFraction, ECsvCustomStatOp::Set);
 	}
 #endif
+
+#if !UE_BUILD_SHIPPING
+	static TMap<uint32, TArray<FString>>* CsvFrameExecCmds = NULL;
+	if (CsvFrameExecCmds == NULL)
+	{
+		CsvFrameExecCmds = new TMap<uint32, TArray<FString>>();
+
+		FString CsvExecCommandsStr;
+		FParse::Value(FCommandLine::Get(), TEXT("-csvExecCmds="), CsvExecCommandsStr, false);
+
+		TArray<FString> CsvExecCommandsList;
+		if (CsvExecCommandsStr.ParseIntoArray(CsvExecCommandsList, TEXT(","), true) > 0)
+		{
+			for (FString FrameAndCommand : CsvExecCommandsList)
+			{
+				TArray<FString> FrameAndCommandList;
+				if (FrameAndCommand.ParseIntoArray(FrameAndCommandList, TEXT(":"), true) == 2)
+				{
+					uint32 Frame = FCString::Atoi(*FrameAndCommandList[0]);
+					if (!CsvFrameExecCmds->Find(Frame))
+					{
+						CsvFrameExecCmds->Add(Frame, TArray<FString>());
+					}
+					(*CsvFrameExecCmds)[Frame].Add(FrameAndCommandList[1]);
+				}
+			}
+		}
+	}
+	if (FCsvProfiler::Get()->IsCapturing())
+	{
+		TArray<FString>* FrameCommands = CsvFrameExecCmds->Find(FCsvProfiler::Get()->GetCaptureFrameNumber());
+		if (FrameCommands != NULL)
+		{
+			for (FString Cmd : *FrameCommands)
+			{
+				GEngine->Exec(GWorld, *Cmd);
+			}
+		}
+	}
+
+#endif
 }
 
 static void UpdateCoreCsvStats_EndFrame()
