Add logging for csvprofiler execcmds

diff --git a/Engine/Source/Runtime/Launch/Private/LaunchEngineLoop.cpp b/Engine/Source/Runtime/Launch/Private/LaunchEngineLoop.cpp
--- a/Engine/Source/Runtime/Launch/Private/LaunchEngineLoop.cpp
+++ b/Engine/Source/Runtime/Launch/Private/LaunchEngineLoop.cpp
@@ -1325,6 +1325,7 @@ static void UpdateCoreCsvStats_BeginFrame()
 		{
 			for (FString Cmd : *FrameCommands)
 			{
+				CSV_EVENT_GLOBAL(TEXT("Executing CSV exec command : %s"), *Cmd);
 				GEngine->Exec(GWorld, *Cmd);
 			}
 		}
