Add -noexceptionhandler command line parameter. This allows you to use JIT debugging.

diff --git a/Engine/Source/Runtime/Core/Private/Windows/WindowsRunnableThread.cpp b/Engine/Source/Runtime/Core/Private/Windows/WindowsRunnableThread.cpp
--- a/Engine/Source/Runtime/Core/Private/Windows/WindowsRunnableThread.cpp
+++ b/Engine/Source/Runtime/Core/Private/Windows/WindowsRunnableThread.cpp
@@ -15,11 +15,12 @@ uint32 FRunnableThreadWin::GuardedRun()
 	FPlatformProcess::SetThreadAffinityMask(ThreadAffinityMask);
 
 	FPlatformProcess::SetThreadName(*ThreadName);
-
+	const TCHAR* CmdLine = ::GetCommandLineW();
+	bool bNoExceptionHandler = FParse::Param(::GetCommandLineW(), TEXT("noexceptionhandler"));
 #if UE_BUILD_DEBUG
 	if (true && !GAlwaysReportCrash)
 #else
-	if (FPlatformMisc::IsDebuggerPresent() && !GAlwaysReportCrash)
+	if (bNoExceptionHandler || (FPlatformMisc::IsDebuggerPresent() && !GAlwaysReportCrash))
 #endif // UE_BUILD_DEBUG
 	{
 		ExitCode = Run();
diff --git a/Engine/Source/Runtime/Launch/Private/Windows/LaunchWindows.cpp b/Engine/Source/Runtime/Launch/Private/Windows/LaunchWindows.cpp
--- a/Engine/Source/Runtime/Launch/Private/Windows/LaunchWindows.cpp
+++ b/Engine/Source/Runtime/Launch/Private/Windows/LaunchWindows.cpp
@@ -231,6 +231,8 @@ int32 WINAPI WinMain( _In_ HINSTANCE hInInstance, _In_opt_ HINSTANCE hPrevInstan
 	}
 #endif
 
+	bool bNoExceptionHandler = FParse::Param(CmdLine,TEXT("noexceptionhandler"));
+	(void)bNoExceptionHandler;
 	// Using the -noinnerexception parameter will disable the exception handler within native C++, which is call from managed code,
 	// which is called from this function.
 	// The default case is to have three wrapped exception handlers 
@@ -238,7 +240,7 @@ int32 WINAPI WinMain( _In_ HINSTANCE hInInstance, _In_opt_ HINSTANCE hPrevInstan
 	// The inner exception handler in GuardedMainWrapper() catches crashes/asserts in native C++ code and is the only way to get the
 	// correct callstack when running a 64-bit executable. However, XAudio2 sometimes (?) don't like this and it may result in no sound.
 #ifdef _WIN64
-	if ( FParse::Param(CmdLine,TEXT("noinnerexception")) || FApp::IsBenchmarking() )
+	if ( FParse::Param(CmdLine,TEXT("noinnerexception")) || FApp::IsBenchmarking() || bNoExceptionHandler)
 	{
 		GEnableInnerException = false;
 	}
@@ -247,7 +249,7 @@ int32 WINAPI WinMain( _In_ HINSTANCE hInInstance, _In_opt_ HINSTANCE hPrevInstan
 #if UE_BUILD_DEBUG
 	if( true && !GAlwaysReportCrash )
 #else
-	if( FPlatformMisc::IsDebuggerPresent() && !GAlwaysReportCrash )
+	if( bNoExceptionHandler || (FPlatformMisc::IsDebuggerPresent() && !GAlwaysReportCrash ))
 #endif
 	{
 		// Don't use exception handling when a debugger is attached to exactly trap the crash. This does NOT check
diff --git a/Engine/Source/Runtime/RenderCore/Private/RenderingThread.cpp b/Engine/Source/Runtime/RenderCore/Private/RenderingThread.cpp
--- a/Engine/Source/Runtime/RenderCore/Private/RenderingThread.cpp
+++ b/Engine/Source/Runtime/RenderCore/Private/RenderingThread.cpp
@@ -462,7 +462,8 @@ public:
 		FPlatformProcess::SetupRenderThread();
 
 #if PLATFORM_WINDOWS
-		if ( !FPlatformMisc::IsDebuggerPresent() || GAlwaysReportCrash )
+		bool bNoExceptionHandler = FParse::Param(FCommandLine::Get(), TEXT("noexceptionhandler"));
+		if ( !bNoExceptionHandler && (!FPlatformMisc::IsDebuggerPresent() || GAlwaysReportCrash))
 		{
 #if !PLATFORM_SEH_EXCEPTIONS_DISABLED
 			__try
