UE-91648 Hololens suspend/resume not working
-Hololens was set up to do extended execution rather than suspend/resume as a workaround to a very old long since fixed sdk bug.
-Now it can suspend/resume correctly.
-Sleep still forces an app restart on PC, because we do not support suspend/resume with remoting.
-The D3D Trim on suspend is PLATFORM_HOLOLENS out of an abundance of caution.  It may well be correct for windows as well, but if its wrong the result will be hitches in d3d operations when resuming, which seems tricky.

diff --git a/Engine/Plugins/Runtime/WindowsMixedReality/Source/WindowsMixedRealityHMD/Private/WindowsMixedRealityHMD.cpp b/Engine/Plugins/Runtime/WindowsMixedReality/Source/WindowsMixedRealityHMD/Private/WindowsMixedRealityHMD.cpp
--- a/Engine/Plugins/Runtime/WindowsMixedReality/Source/WindowsMixedRealityHMD/Private/WindowsMixedRealityHMD.cpp
+++ b/Engine/Plugins/Runtime/WindowsMixedReality/Source/WindowsMixedRealityHMD/Private/WindowsMixedRealityHMD.cpp
@@ -1927,7 +1927,10 @@ namespace WindowsMixedReality
 	// Prevent crashes if computer goes to sleep.
 	void FWindowsMixedRealityHMD::AppServicePause()
 	{
+		// We support sleep on hololens.
+#if !PLATFORM_HOLOLENS
 		this->bRequestRestart = true;
+#endif
 	}
 
 	void FWindowsMixedRealityHMD::AppServiceResume()
diff --git a/Engine/Source/Runtime/Launch/Private/HoloLens/HoloLensLaunch.cpp b/Engine/Source/Runtime/Launch/Private/HoloLens/HoloLensLaunch.cpp
--- a/Engine/Source/Runtime/Launch/Private/HoloLens/HoloLensLaunch.cpp
+++ b/Engine/Source/Runtime/Launch/Private/HoloLens/HoloLensLaunch.cpp
@@ -24,11 +24,13 @@
 #include "HoloLens/HoloLensWindow.h"
 #include <stdio.h>
 
+#include "HoloLens/AllowWindowsPlatformAtomics.h"
 #include "HoloLens/AllowWindowsPlatformTypes.h"
 #include <collection.h>
 #include <ppltasks.h>
 #include <concurrent_queue.h>
 #include "HoloLens/HideWindowsPlatformTypes.h"
+#include "HoloLens/HideWindowsPlatformAtomics.h"
 
 // @MIXEDREALITY_CHANGE : BEGIN - Enable Stereo rendering for suspended apps
 #include "IHeadMountedDisplay.h"
@@ -103,22 +105,17 @@ public:
 
 	void InitOptionalPackages();
 
-protected:
-	// Registers the extended execution session
-	void RequestExtendedExecution();
-	
-	// Extended execution handlers.
-	void OnExtensionRevokedHandler(Platform::Object^ obj, ExtendedExecutionRevokedEventArgs^ args);
 private:
 	Windows::Foundation::Rect existingSize;
 
-	// If the ExtendedExecutionSession goes out of scope any pending request will be nullified
-	ExtendedExecutionSession^ ExecutionSession;
 	bool ActivationComplete;
 
 	void OnActivated( _In_ Windows::ApplicationModel::Core::CoreApplicationView^ applicationView, _In_ Windows::ApplicationModel::Activation::IActivatedEventArgs^ args );
 	void OnResuming( _In_ Platform::Object^ sender, _In_ Platform::Object^ args );
 	void OnSuspending( _In_ Platform::Object^ sender, _In_ Windows::ApplicationModel::SuspendingEventArgs^ args );
+	void OnEnteredBackground(_In_ Platform::Object^ sender, _In_ Windows::ApplicationModel::EnteredBackgroundEventArgs^ args);
+	void OnLeavingBackground(_In_ Platform::Object^ sender, _In_ Windows::ApplicationModel::LeavingBackgroundEventArgs^ args);
+	void OnExiting(_In_ Platform::Object^ sender, _In_ Platform::Object^ args);
 
 	const TCHAR* GetPointerUpdateKindString(Windows::UI::Input::PointerUpdateKind InKind);
 	EMouseButtons::Type PointerUpdateKindToUEKey(Windows::UI::Input::PointerUpdateKind InKind, bool& bWasPressed);
@@ -162,7 +159,6 @@ public:
 };
 
 ViewProvider::ViewProvider() :
-	ExecutionSession(nullptr),
 	ActivationComplete(false)
 {
 }
@@ -586,7 +582,14 @@ void ViewProvider::ProcessEvents()
 
 		if (nullptr != Dispatcher)
 		{
-			Dispatcher->ProcessEvents( CoreProcessEventsOption::ProcessAllIfPresent );
+			if (bVisible)
+			{
+				Dispatcher->ProcessEvents(CoreProcessEventsOption::ProcessAllIfPresent);
+			}
+			else
+			{
+				Dispatcher->ProcessEvents(CoreProcessEventsOption::ProcessOneAndAllPending);
+			}
 		}
 	}
 
@@ -613,10 +616,79 @@ void ViewProvider::Initialize(Windows::ApplicationModel::Core::CoreApplicationVi
 	applicationView->Activated += ref new Windows::Foundation::TypedEventHandler< CoreApplicationView^, IActivatedEventArgs^ >( this, &ViewProvider::OnActivated );
 	CoreApplication::Suspending += ref new Windows::Foundation::EventHandler< Windows::ApplicationModel::SuspendingEventArgs^ >( this, &ViewProvider::OnSuspending );
 	CoreApplication::Resuming += ref new Windows::Foundation::EventHandler< Platform::Object^>( this, & ViewProvider::OnResuming );
+	CoreApplication::EnteredBackground += ref new Windows::Foundation::EventHandler< Windows::ApplicationModel::EnteredBackgroundEventArgs^ >(this, &ViewProvider::OnEnteredBackground);
+	CoreApplication::LeavingBackground += ref new Windows::Foundation::EventHandler< Windows::ApplicationModel::LeavingBackgroundEventArgs^ >(this, &ViewProvider::OnLeavingBackground);
+	CoreApplication::Exiting += ref new Windows::Foundation::EventHandler< Platform::Object^>(this, &ViewProvider::OnExiting);
+
+}
+
+static void InitCommandLine(Platform::String^ args)
+{
+	// check for a ue4commandline.txt.
+	FCommandLine::Set(L"");
+
+	// FPaths should be good at this point (FPaths::EngineDir() is used during static init for instance)
+	// so be specific about the location we want, otherwise working directory changes can mess this up.
+	// Also this means the regular file interface is in place, so let's use that.
+	FString FileName = FString(FPaths::RootDir()) / TEXT("UE4CommandLine.txt");
+
+	IFileHandle* CmdLineFileHandle = IPlatformFile::GetPlatformPhysical().OpenRead(*FileName);
+	if (CmdLineFileHandle != nullptr)
+	{
+		const int64 MaxCmdLineSize = 65536;
+
+		int64 CmdLineFileSize = CmdLineFileHandle->Size();
+		if (CmdLineFileSize < MaxCmdLineSize - 1)
+		{
+			char AnsiCmdLine[MaxCmdLineSize] = {};
+			if (CmdLineFileHandle->Read(reinterpret_cast<uint8*>(AnsiCmdLine), CmdLineFileSize))
+			{
+				FString CmdLine = StringCast<TCHAR>(AnsiCmdLine).Get();
+				CmdLine.TrimStartAndEndInline();
+				
+				if (args != nullptr)
+				{
+					CmdLine += FString(" ");
+					CmdLine += FString(args->Data());
+
+					CmdLine.TrimStartAndEndInline();
+				}
+
+				UE_LOG(LogLaunchHoloLens, Log, TEXT("%s"), *CmdLine);
+				FCommandLine::Set(*CmdLine);
+			}
+			else
+			{
+				UE_LOG(LogLaunchHoloLens, Warning, TEXT("Failed to read commandline from %s!"), *FileName);
+			}
+		}
+		else
+		{
+			UE_LOG(LogLaunchHoloLens, Warning, TEXT("Commandline file %s too large (%d / %d bytes).  Ignoring."), *FileName, CmdLineFileSize, MaxCmdLineSize);
+		}
+		delete CmdLineFileHandle;
+	}
 }
 
 void ViewProvider::OnActivated(_In_ Windows::ApplicationModel::Core::CoreApplicationView^ applicationView, _In_ Windows::ApplicationModel::Activation::IActivatedEventArgs^ args)
 {
+	// Check for launch activation
+	Platform::String^ arg = nullptr;
+	if (args->Kind == ActivationKind::Launch)
+	{
+		auto launchArgs = static_cast<LaunchActivatedEventArgs^>(args);
+		arg = launchArgs->Arguments;
+	}
+
+	// Check for protocol activation
+	if (args->Kind == ActivationKind::Protocol)
+	{
+		auto protocolArgs = static_cast<ProtocolActivatedEventArgs^>(args);
+		arg = protocolArgs->Uri->ToString();
+	}
+
+	InitCommandLine(arg);
+	
 	// Take advantage of this opportunity to measure the desktop (based on view bounds before window activation).
 	FHoloLensApplication::CacheDesktopSize();
 
@@ -625,9 +697,6 @@ void ViewProvider::OnActivated(_In_ Windows::ApplicationModel::Core::CoreApplica
 	window->Activate();
     ActivationComplete = true;
 
-	// Initial registration of the extended execution
-	RequestExtendedExecution();
-
 	// Query this to find out if the application was shut down gracefully last time
 	if (args != nullptr)
 	{
@@ -670,16 +739,7 @@ void ViewProvider::OnResuming(_In_ Platform::Object^ Sender, _In_ Platform::Obje
 	RHIResumeRendering();
 
 	// Notify application of resume
-	FCoreDelegates::ApplicationHasEnteredForegroundDelegate.Broadcast();
-	// Invalidate pending execution request
-	ExecutionSession = nullptr;
-
-	// Re-register the extended execution
-	RequestExtendedExecution();
-
-	// @MIXEDREALITY_CHANGE : BEGIN
-	GEngine->StereoRenderingDevice->EnableStereo(true);
-	// @MIXEDREALITY_CHANGE : END
+	FCoreDelegates::ApplicationHasReactivatedDelegate.Broadcast();
 }
 
 void ViewProvider::OnSuspending(_In_ Platform::Object^ Sender, _In_ Windows::ApplicationModel::SuspendingEventArgs^ Args)
@@ -688,7 +748,7 @@ void ViewProvider::OnSuspending(_In_ Platform::Object^ Sender, _In_ Windows::App
 	Windows::ApplicationModel::SuspendingDeferral^ SuspendingEvent = Args->SuspendingOperation->GetDeferral();
 
 	// Notify application of suspend. Application should kick off an async save at this point.
-	FCoreDelegates::ApplicationWillEnterBackgroundDelegate.Broadcast();
+	FCoreDelegates::ApplicationWillDeactivateDelegate.Broadcast();
 
 	// Flush the RenderingThread
 	FlushRenderingCommands();
@@ -708,6 +768,32 @@ void ViewProvider::OnSuspending(_In_ Platform::Object^ Sender, _In_ Windows::App
 	SuspendingEvent->Complete();
 }
 
+void ViewProvider::OnEnteredBackground(_In_ Platform::Object^ sender, _In_ Windows::ApplicationModel::EnteredBackgroundEventArgs^ args)
+{
+	if (!GIsRunning)
+	{
+		return;
+	}
+
+	FCoreDelegates::ApplicationWillEnterBackgroundDelegate.Broadcast();
+}
+
+void ViewProvider::OnLeavingBackground(_In_ Platform::Object^ sender, _In_ Windows::ApplicationModel::LeavingBackgroundEventArgs^ args)
+{
+	if (!GIsRunning)
+	{
+		return;
+	}
+
+	FCoreDelegates::ApplicationHasEnteredForegroundDelegate.Broadcast();
+}
+
+void ViewProvider::OnExiting(_In_ Platform::Object^ sender, _In_ Platform::Object^ args)
+{
+	FCoreDelegates::ApplicationWillTerminateDelegate.Broadcast();
+}
+
+
 const TCHAR* ViewProvider::GetPointerUpdateKindString(Windows::UI::Input::PointerUpdateKind InKind)
 {
 	switch (InKind)
@@ -739,52 +825,6 @@ const TCHAR* ViewProvider::GetPointerUpdateKindString(Windows::UI::Input::Pointe
 	return TEXT("*** UNKNOWN ***");
 }
 
-void ViewProvider::OnExtensionRevokedHandler(Platform::Object^ obj, ExtendedExecutionRevokedEventArgs^ args)
-{
-	UE_LOG(LogLaunchHoloLens, Warning, TEXT("Extended execution revoked.  Reason: %d"), static_cast<uint32>(args->Reason) );
-
-	// Only request extended exeuction when it is resumed
-	if (args->Reason == ExtendedExecutionRevokedReason::Resumed)
-	{
-		// Invalidate pending execution request
-		ExecutionSession = nullptr;
-
-		// Re-register the extended execution
-		RequestExtendedExecution();
-	}
-}
-
-void ViewProvider::RequestExtendedExecution()
-{
-	UE_LOG(LogLaunchHoloLens, Log, TEXT("Requesting Extended Execution Session.") );
-
-	if (ExecutionSession == nullptr)
-	{
-		ExecutionSession = ref new ExtendedExecutionSession();
-		ExecutionSession->Reason = ExtendedExecutionReason::Unspecified;
-		ExecutionSession->Description = "None";
-		ExecutionSession->Revoked += ref new TypedEventHandler<Platform::Object^, ExtendedExecutionRevokedEventArgs^>(this, &ViewProvider::OnExtensionRevokedHandler);
-	}
-
-	Concurrency::create_task(ExecutionSession->RequestExtensionAsync())
-		.then([this](Concurrency::task<ExtendedExecutionResult> t)
-	{
-		try
-		{
-			// At this point the request has been made. 
-			// When the IAsyncOperation request completes verify that 
-			// the ExtendedExecutionResult == Allowed and  you will not 
-			// suspend for up to 10 minutes while minimized.
-
-			ExtendedExecutionResult result = t.get();
-			UE_LOG(LogLaunchHoloLens, Log, TEXT("Extended Execution Result: %d"), static_cast<uint32>(result) );
-		}
-		catch (...)
-		{
-			UE_LOG(LogLaunchHoloLens, Warning, TEXT("Extended Execution Denied (Exception).") );
-		}
-	});
-}
 
 EMouseButtons::Type ViewProvider::PointerUpdateKindToUEKey(Windows::UI::Input::PointerUpdateKind InKind, bool& bWasPressed)
 {
@@ -996,6 +1036,7 @@ void ViewProvider::InitOptionalPackages()
 #endif
 }
 
+#include "AllowWindowsPlatformAtomics.h"
 #include "AllowWindowsPlatformTypes.h"
 #include "NetworkMessage.h"
 
@@ -1046,55 +1087,15 @@ void ViewProvider::Run()
 	Windows::ApplicationModel::Core::CoreApplication::Exit();
 }
 #include "HideWindowsPlatformTypes.h"
+#include "HideWindowsPlatformAtomics.h"
 
 void ViewProvider::Uninitialize()
 {
 }
 
-static void InitCommandLine()
-{
-	// check for a ue4commandline.txt.
-	FCommandLine::Set(L"");
-
-	// FPaths should be good at this point (FPaths::EngineDir() is used during static init for instance)
-	// so be specific about the location we want, otherwise working directory changes can mess this up.
-	// Also this means the regular file interface is in place, so let's use that.
-	FString FileName = FString(FPaths::RootDir()) / TEXT("UE4CommandLine.txt");
-
-	IFileHandle* CmdLineFileHandle = IPlatformFile::GetPlatformPhysical().OpenRead(*FileName);
-	if (CmdLineFileHandle != nullptr)
-	{
-		const int64 MaxCmdLineSize = 65536;
-
-		int64 CmdLineFileSize = CmdLineFileHandle->Size();
-		if (CmdLineFileSize < MaxCmdLineSize - 1)
-		{
-			char AnsiCmdLine[MaxCmdLineSize] = {};
-			if (CmdLineFileHandle->Read(reinterpret_cast<uint8*>(AnsiCmdLine), CmdLineFileSize))
-			{
-				FString CmdLine = StringCast<TCHAR>(AnsiCmdLine).Get();
-				CmdLine.TrimStartAndEndInline();
-				UE_LOG(LogLaunchHoloLens, Log, TEXT("%s"), *CmdLine);
-				FCommandLine::Set(*CmdLine);
-			}
-			else
-			{
-				UE_LOG(LogLaunchHoloLens, Warning, TEXT("Failed to read commandline from %s!"), *FileName);
-			}
-		}
-		else
-		{
-			UE_LOG(LogLaunchHoloLens, Warning, TEXT("Commandline file %s too large (%d / %d bytes).  Ignoring."), *FileName, CmdLineFileSize, MaxCmdLineSize);
-		}
-		delete CmdLineFileHandle;
-	}
-}
-
 [Platform::MTAThread]
 int main(Platform::Array<Platform::String^>^)
 {
-	InitCommandLine();
-
 	uint64 GameThreadAffinity = FPlatformAffinity::GetMainGameMask();
 	FPlatformProcess::SetThreadAffinityMask((DWORD_PTR)GameThreadAffinity);
 
diff --git a/Engine/Source/Runtime/Windows/D3D11RHI/Private/D3D11Device.cpp b/Engine/Source/Runtime/Windows/D3D11RHI/Private/D3D11Device.cpp
--- a/Engine/Source/Runtime/Windows/D3D11RHI/Private/D3D11Device.cpp
+++ b/Engine/Source/Runtime/Windows/D3D11RHI/Private/D3D11Device.cpp
@@ -265,6 +265,20 @@ void FD3D11DynamicRHI::RHIPushEvent(const TCHAR* Name, FColor Color)
 	GPUProfilingData.PushEvent(Name, Color);
 }
 
+#if PLATFORM_HOLOLENS
+void FD3D11DynamicRHI::RHISuspendRendering()
+{
+	TRefCountPtr<IDXGIDevice3> dxgiDevice;
+	if (Direct3DDevice->QueryInterface(IID_PPV_ARGS(dxgiDevice.GetInitReference())))
+	{
+		dxgiDevice->Trim();
+	}
+}
+void FD3D11DynamicRHI::RHIResumeRendering()
+{
+}
+#endif
+
 void FD3D11DynamicRHI::RHIPopEvent()
 { 
 	GPUProfilingData.PopEvent(); 
diff --git a/Engine/Source/Runtime/Windows/D3D11RHI/Private/D3D11RHIPrivate.h b/Engine/Source/Runtime/Windows/D3D11RHI/Private/D3D11RHIPrivate.h
--- a/Engine/Source/Runtime/Windows/D3D11RHI/Private/D3D11RHIPrivate.h
+++ b/Engine/Source/Runtime/Windows/D3D11RHI/Private/D3D11RHIPrivate.h
@@ -380,6 +380,11 @@ public:
 	virtual void PostInit() override;
 	virtual void Shutdown() override;
 	virtual const TCHAR* GetName() override { return TEXT("D3D11"); }
+	
+#if PLATFORM_HOLOLENS
+	virtual void RHISuspendRendering() override;
+	virtual void RHIResumeRendering() override;
+#endif
 
 	// HDR display output
 	virtual void EnableHDR();
