UE-91752 //UE4/Release-4.25 - Compile UE4Game HoloLens - winnt.h, Warning: 'InterlockedAdd': macro redefinition
-Removed include which defines these macros, they are defined by the windows 10 sdk on hololens and seem to be unnecessary.

diff --git a/Engine/Source/Runtime/Launch/Private/HoloLens/HoloLensLaunch.cpp b/Engine/Source/Runtime/Launch/Private/HoloLens/HoloLensLaunch.cpp
--- a/Engine/Source/Runtime/Launch/Private/HoloLens/HoloLensLaunch.cpp
+++ b/Engine/Source/Runtime/Launch/Private/HoloLens/HoloLensLaunch.cpp
@@ -24,13 +24,13 @@
 #include "HoloLens/HoloLensWindow.h"
 #include <stdio.h>
 
-#include "HoloLens/AllowWindowsPlatformAtomics.h"
+//#include "HoloLens/AllowWindowsPlatformAtomics.h"
 #include "HoloLens/AllowWindowsPlatformTypes.h"
 #include <collection.h>
 #include <ppltasks.h>
 #include <concurrent_queue.h>
 #include "HoloLens/HideWindowsPlatformTypes.h"
-#include "HoloLens/HideWindowsPlatformAtomics.h"
+//#include "HoloLens/HideWindowsPlatformAtomics.h"
 
 // @MIXEDREALITY_CHANGE : BEGIN - Enable Stereo rendering for suspended apps
 #include "IHeadMountedDisplay.h"
