Add a core delegate called on sampling the input (UE-91807).

- Allows to generalize input latency tracking.

diff --git a/Engine/Source/Runtime/Core/Private/Misc/CoreDelegates.cpp b/Engine/Source/Runtime/Core/Private/Misc/CoreDelegates.cpp
--- a/Engine/Source/Runtime/Core/Private/Misc/CoreDelegates.cpp
+++ b/Engine/Source/Runtime/Core/Private/Misc/CoreDelegates.cpp
@@ -89,6 +89,7 @@ FSimpleMulticastDelegate FCoreDelegates::OnEnginePreExit;
 FCoreDelegates::FGatherAdditionalLocResPathsDelegate FCoreDelegates::GatherAdditionalLocResPathsCallback;
 FSimpleMulticastDelegate FCoreDelegates::ColorPickerChanged;
 FSimpleMulticastDelegate FCoreDelegates::OnBeginFrame;
+FSimpleMulticastDelegate FCoreDelegates::OnSamplingInput;
 FSimpleMulticastDelegate FCoreDelegates::OnEndFrame;
 FSimpleMulticastDelegate FCoreDelegates::OnBeginFrameRT;
 FSimpleMulticastDelegate FCoreDelegates::OnEndFrameRT;
diff --git a/Engine/Source/Runtime/Core/Public/Misc/CoreDelegates.h b/Engine/Source/Runtime/Core/Public/Misc/CoreDelegates.h
--- a/Engine/Source/Runtime/Core/Public/Misc/CoreDelegates.h
+++ b/Engine/Source/Runtime/Core/Public/Misc/CoreDelegates.h
@@ -267,6 +267,9 @@ public:
 	// Called at the beginning of a frame
 	static FSimpleMulticastDelegate OnBeginFrame;
 
+	// Called at the moment of sampling the input (currently on the gamethread)
+	static FSimpleMulticastDelegate OnSamplingInput;
+
 	// Called at the end of a frame
 	static FSimpleMulticastDelegate OnEndFrame;
 
diff --git a/Engine/Source/Runtime/Launch/Private/LaunchEngineLoop.cpp b/Engine/Source/Runtime/Launch/Private/LaunchEngineLoop.cpp
--- a/Engine/Source/Runtime/Launch/Private/LaunchEngineLoop.cpp
+++ b/Engine/Source/Runtime/Launch/Private/LaunchEngineLoop.cpp
@@ -1394,6 +1394,10 @@ namespace AppLifetimeEventCapture
 }
 #endif //WITH_ENGINE
 
+static void UpdateGInputTime()
+{
+	GInputTime = FPlatformTime::Cycles64();
+}
 
 DECLARE_CYCLE_STAT(TEXT("FEngineLoop::PreInitPreStartupScreen.AfterStats"), STAT_FEngineLoop_PreInitPreStartupScreen_AfterStats, STATGROUP_LoadTime);
 DECLARE_CYCLE_STAT(TEXT("FEngineLoop::PreInitPostStartupScreen.AfterStats"), STAT_FEngineLoop_PreInitPostStartupScreen_AfterStats, STATGROUP_LoadTime);
@@ -1537,6 +1541,7 @@ int32 FEngineLoop::PreInitPreStartupScreen(const TCHAR* CmdLine)
 #if WITH_ENGINE
 	FCoreUObjectDelegates::PostGarbageCollectConditionalBeginDestroy.AddStatic(DeferredPhysResourceCleanup);
 #endif
+	FCoreDelegates::OnSamplingInput.AddStatic(UpdateGInputTime);
 
 #if defined(WITH_LAUNCHERCHECK) && WITH_LAUNCHERCHECK
 	if (ILauncherCheckModule::Get().WasRanFromLauncher() == false)
@@ -4846,7 +4851,7 @@ void FEngineLoop::Tick()
 		FPlatformFileManager::Get().TickActivePlatformFile();
 
 		// Roughly track the time when the input was sampled
-		GInputTime = FPlatformTime::Cycles64();
+		FCoreDelegates::OnSamplingInput.Broadcast();
 
 		// process accumulated Slate input
 		if (FSlateApplication::IsInitialized() && !bIdleMode)
