Workaround for IOS failing to load a default material until we can fix the shader being in the wrong library

diff --git a/Engine/Source/Runtime/Launch/Private/LaunchEngineLoop.cpp b/Engine/Source/Runtime/Launch/Private/LaunchEngineLoop.cpp
--- a/Engine/Source/Runtime/Launch/Private/LaunchEngineLoop.cpp
+++ b/Engine/Source/Runtime/Launch/Private/LaunchEngineLoop.cpp
@@ -3053,6 +3053,16 @@ int32 FEngineLoop::PreInitPostStartupScreen(const TCHAR* CmdLine)
 				FShaderPipelineCache::OpenPipelineFileCache(GMaxRHIShaderPlatform);
 			}
 		}
+		else
+		{
+			// TEMPORARY FIX UNTIL WE FIGURE OUT WHY A DEFAULT MATERIAL IS USING A GAME LIBRARY SHADER
+			// Open the game library which contains the material shaders.
+			FShaderCodeLibrary::OpenLibrary(FApp::GetProjectName(), FPaths::ProjectContentDir());
+			for (const FString& RootDir : FPlatformMisc::GetAdditionalRootDirectories())
+			{
+				FShaderCodeLibrary::OpenLibrary(FApp::GetProjectName(), FPaths::Combine(RootDir, FApp::GetProjectName(), TEXT("Content")));
+			}
+		}
 
 		InitGameTextLocalization();
 
