Use the default editor target config setting, if it exists, to determine the path to the editor target file for startup re-compilation purposes

diff --git a/Engine/Source/Runtime/Launch/Private/LaunchEngineLoop.cpp b/Engine/Source/Runtime/Launch/Private/LaunchEngineLoop.cpp
--- a/Engine/Source/Runtime/Launch/Private/LaunchEngineLoop.cpp
+++ b/Engine/Source/Runtime/Launch/Private/LaunchEngineLoop.cpp
@@ -5444,19 +5444,28 @@ bool FEngineLoop::AppInit( )
 
 		// Find the editor target
 		FString EditorTargetFileName;
-		for (const FTargetInfo& Target : FDesktopPlatformModule::Get()->GetTargetsForProject(FPaths::GetProjectFilePath()))
+
+		// Config files can specify the default editor target
+		if (GConfig->GetString(TEXT("/Script/BuildSettings.BuildSettings"), TEXT("DefaultEditorTarget"), EditorTargetFileName, GEngineIni))
+		{
+			EditorTargetFileName = FTargetReceipt::GetDefaultPath(FPlatformMisc::ProjectDir(), *EditorTargetFileName, FPlatformProcess::GetBinariesSubdirectory(), FApp::GetBuildConfiguration(), nullptr);
+		}
+		else
 		{
-			if (Target.Type == EBuildTargetType::Editor)
+			for (const FTargetInfo& Target : FDesktopPlatformModule::Get()->GetTargetsForProject(FPaths::GetProjectFilePath()))
 			{
-				if(FPaths::IsUnderDirectory(Target.Path, FPlatformMisc::ProjectDir()))
-				{
-					EditorTargetFileName = FTargetReceipt::GetDefaultPath(FPlatformMisc::ProjectDir(), *Target.Name, FPlatformProcess::GetBinariesSubdirectory(), FApp::GetBuildConfiguration(), nullptr);
-				}
-				else if(FPaths::IsUnderDirectory(Target.Path, FPaths::EngineDir()))
+				if (Target.Type == EBuildTargetType::Editor)
 				{
-					EditorTargetFileName = FTargetReceipt::GetDefaultPath(*FPaths::EngineDir(), *Target.Name, FPlatformProcess::GetBinariesSubdirectory(), FApp::GetBuildConfiguration(), nullptr);
+					if (FPaths::IsUnderDirectory(Target.Path, FPlatformMisc::ProjectDir()))
+					{
+						EditorTargetFileName = FTargetReceipt::GetDefaultPath(FPlatformMisc::ProjectDir(), *Target.Name, FPlatformProcess::GetBinariesSubdirectory(), FApp::GetBuildConfiguration(), nullptr);
+					}
+					else if (FPaths::IsUnderDirectory(Target.Path, FPaths::EngineDir()))
+					{
+						EditorTargetFileName = FTargetReceipt::GetDefaultPath(*FPaths::EngineDir(), *Target.Name, FPlatformProcess::GetBinariesSubdirectory(), FApp::GetBuildConfiguration(), nullptr);
+					}
+					break;
 				}
-				break;
 			}
 		}
 
