Minor update to editor target selection on launch to account for the possibility of default targets living in the engine

diff --git a/Engine/Source/Runtime/Launch/Private/LaunchEngineLoop.cpp b/Engine/Source/Runtime/Launch/Private/LaunchEngineLoop.cpp
--- a/Engine/Source/Runtime/Launch/Private/LaunchEngineLoop.cpp
+++ b/Engine/Source/Runtime/Launch/Private/LaunchEngineLoop.cpp
@@ -5444,28 +5444,22 @@ bool FEngineLoop::AppInit( )
 
 		// Find the editor target
 		FString EditorTargetFileName;
+		FString DefaultEditorTarget;
+		GConfig->GetString(TEXT("/Script/BuildSettings.BuildSettings"), TEXT("DefaultEditorTarget"), DefaultEditorTarget, GEngineIni);
 
-		// Config files can specify the default editor target
-		if (GConfig->GetString(TEXT("/Script/BuildSettings.BuildSettings"), TEXT("DefaultEditorTarget"), EditorTargetFileName, GEngineIni))
+		for (const FTargetInfo& Target : FDesktopPlatformModule::Get()->GetTargetsForProject(FPaths::GetProjectFilePath()))
 		{
-			EditorTargetFileName = FTargetReceipt::GetDefaultPath(FPlatformMisc::ProjectDir(), *EditorTargetFileName, FPlatformProcess::GetBinariesSubdirectory(), FApp::GetBuildConfiguration(), nullptr);
-		}
-		else
-		{
-			for (const FTargetInfo& Target : FDesktopPlatformModule::Get()->GetTargetsForProject(FPaths::GetProjectFilePath()))
+			if (Target.Type == EBuildTargetType::Editor && (DefaultEditorTarget.Len() == 0 || Target.Name == DefaultEditorTarget))
 			{
-				if (Target.Type == EBuildTargetType::Editor)
+				if (FPaths::IsUnderDirectory(Target.Path, FPlatformMisc::ProjectDir()))
 				{
-					if (FPaths::IsUnderDirectory(Target.Path, FPlatformMisc::ProjectDir()))
-					{
-						EditorTargetFileName = FTargetReceipt::GetDefaultPath(FPlatformMisc::ProjectDir(), *Target.Name, FPlatformProcess::GetBinariesSubdirectory(), FApp::GetBuildConfiguration(), nullptr);
-					}
-					else if (FPaths::IsUnderDirectory(Target.Path, FPaths::EngineDir()))
-					{
-						EditorTargetFileName = FTargetReceipt::GetDefaultPath(*FPaths::EngineDir(), *Target.Name, FPlatformProcess::GetBinariesSubdirectory(), FApp::GetBuildConfiguration(), nullptr);
-					}
-					break;
+					EditorTargetFileName = FTargetReceipt::GetDefaultPath(FPlatformMisc::ProjectDir(), *Target.Name, FPlatformProcess::GetBinariesSubdirectory(), FApp::GetBuildConfiguration(), nullptr);
 				}
+				else if (FPaths::IsUnderDirectory(Target.Path, FPaths::EngineDir()))
+				{
+					EditorTargetFileName = FTargetReceipt::GetDefaultPath(*FPaths::EngineDir(), *Target.Name, FPlatformProcess::GetBinariesSubdirectory(), FApp::GetBuildConfiguration(), nullptr);
+				}
+				break;
 			}
 		}
 
