Update version number to 5.0

diff --git a/Engine/Build/Build.version b/Engine/Build/Build.version
--- a/Engine/Build/Build.version
+++ b/Engine/Build/Build.version
@@ -1,10 +1,10 @@
 {
-	"MajorVersion": 4,
-	"MinorVersion": 26,
+	"MajorVersion": 5,
+	"MinorVersion": 0,
 	"PatchVersion": 0,
 	"Changelist": 0,
 	"CompatibleChangelist": 0,
 	"IsLicenseeVersion": 0,
 	"IsPromotedBuild": 0,
-	"BranchName": "UE4"
+	"BranchName": "UE5"
 }
diff --git a/Engine/Source/Runtime/Launch/Resources/Version.h b/Engine/Source/Runtime/Launch/Resources/Version.h
--- a/Engine/Source/Runtime/Launch/Resources/Version.h
+++ b/Engine/Source/Runtime/Launch/Resources/Version.h
@@ -54,8 +54,8 @@
 
 // These numbers define the banner UE4 version, and are the most significant numbers when ordering two engine versions (that is, a 4.12.* version is always 
 // newer than a 4.11.* version, regardless of the changelist that it was built with)
-#define ENGINE_MAJOR_VERSION	4
-#define ENGINE_MINOR_VERSION	26
+#define ENGINE_MAJOR_VERSION	5
+#define ENGINE_MINOR_VERSION	0
 #define ENGINE_PATCH_VERSION	0
 
 // Macros for encoding strings
