Fix LaunchCorrectEditorExecutable when using Unique BuildEnvironment.

diff --git a/Engine/Source/Runtime/Launch/Private/LaunchEngineLoop.cpp b/Engine/Source/Runtime/Launch/Private/LaunchEngineLoop.cpp
--- a/Engine/Source/Runtime/Launch/Private/LaunchEngineLoop.cpp
+++ b/Engine/Source/Runtime/Launch/Private/LaunchEngineLoop.cpp
@@ -5245,6 +5245,14 @@ static void CheckForPrintTimesOverride()
 }
 
 #if UE_EDITOR
+//Standardize paths when deciding  if running the proper editor exe. 
+void CleanUpPath(FString& InPath)
+{
+	//Converts to full path will also replace '\' with '/' and will collapse relative directories (C:\foo\..\bar to C:\bar)
+	InPath = FPaths::ConvertRelativePathToFull(InPath);
+	FPaths::RemoveDuplicateSlashes(InPath);
+}
+
 bool LaunchCorrectEditorExecutable(const FString& EditorTargetFileName)
 {
 	// Don't allow relaunching the executable if we're running some unattended scripted process.
@@ -5268,7 +5276,7 @@ bool LaunchCorrectEditorExecutable(const FString& EditorTargetFileName)
 		}
 		LaunchExecutableName = Receipt.Launch;
 	}
-	FPaths::MakeStandardFilename(LaunchExecutableName);
+	CleanUpPath(LaunchExecutableName);
 
 	// Get the current executable name. Don't allow relaunching if we're running the console app.
 	FString CurrentExecutableName = FPlatformProcess::ExecutablePath();
@@ -5276,7 +5284,7 @@ bool LaunchCorrectEditorExecutable(const FString& EditorTargetFileName)
 	{
 		return false;
 	}
-	FPaths::MakeStandardFilename(CurrentExecutableName);
+	CleanUpPath(CurrentExecutableName);
 
 	// Nothing to do if they're the same
 	if(LaunchExecutableName == CurrentExecutableName)
