Fixing build issue with missing intermediary directory.

diff --git a/Engine/Source/Runtime/Launch/Private/LaunchEngineLoop.cpp b/Engine/Source/Runtime/Launch/Private/LaunchEngineLoop.cpp
--- a/Engine/Source/Runtime/Launch/Private/LaunchEngineLoop.cpp
+++ b/Engine/Source/Runtime/Launch/Private/LaunchEngineLoop.cpp
@@ -1748,8 +1748,20 @@ int32 FEngineLoop::PreInitPreStartupScreen(const TCHAR* CmdLine)
 
 	{
 		IPlatformFile& PlatformFile = FPlatformFileManager::Get().GetPlatformFile();
+
+		bool bCreateIntermediateSuccess = PlatformFile.CreateDirectory(*FPaths::ProjectIntermediateDir());
+		if (!bCreateIntermediateSuccess)
+		{
+			UE_LOG(LogInit,Fatal,TEXT("Failed to create Intermediate directory."));
+		}
+
 		FString AutogenAbsolutePath = FPaths::ConvertRelativePathToFull(FPaths::ProjectIntermediateDir() / TEXT("ShaderAutogen"));
-		PlatformFile.CreateDirectory(AutogenAbsolutePath.GetCharArray().GetData());
+		bool bCreateAutogenSuccess = PlatformFile.CreateDirectory(*AutogenAbsolutePath);
+		if (!bCreateAutogenSuccess)
+		{
+			UE_LOG(LogInit,Fatal,TEXT("Failed to create Intermediate/ShaderAutogen/ directory. Make sure Intermediate exists."));
+		}
+
 		AddShaderSourceDirectoryMapping(TEXT("/ShaderAutogen"), AutogenAbsolutePath);
 	}
 
