Fixing issue with creating directory when WITH_EDITOR is disabled.

diff --git a/Engine/Source/Runtime/Launch/Private/LaunchEngineLoop.cpp b/Engine/Source/Runtime/Launch/Private/LaunchEngineLoop.cpp
--- a/Engine/Source/Runtime/Launch/Private/LaunchEngineLoop.cpp
+++ b/Engine/Source/Runtime/Launch/Private/LaunchEngineLoop.cpp
@@ -1746,6 +1746,7 @@ int32 FEngineLoop::PreInitPreStartupScreen(const TCHAR* CmdLine)
 	// Add the default engine shader dir
 	AddShaderSourceDirectoryMapping(TEXT("/Engine"), FGenericPlatformProcess::ShaderDir());
 
+#if WITH_EDITOR
 	{
 		IPlatformFile& PlatformFile = FPlatformFileManager::Get().GetPlatformFile();
 
@@ -1764,6 +1765,7 @@ int32 FEngineLoop::PreInitPreStartupScreen(const TCHAR* CmdLine)
 
 		AddShaderSourceDirectoryMapping(TEXT("/ShaderAutogen"), AutogenAbsolutePath);
 	}
+#endif
 
 	TArray<FString> Tokens;
 	TArray<FString> Switches;
