Tagging main thread as game thread on Android (similar to CL 14378525 for HL2)

diff --git a/Engine/Source/Runtime/Launch/Private/Android/LaunchAndroid.cpp b/Engine/Source/Runtime/Launch/Private/Android/LaunchAndroid.cpp
--- a/Engine/Source/Runtime/Launch/Private/Android/LaunchAndroid.cpp
+++ b/Engine/Source/Runtime/Launch/Private/Android/LaunchAndroid.cpp
@@ -740,6 +740,8 @@ struct android_app* GNativeAndroidApp = NULL;
 
 void android_main(struct android_app* state)
 {
+	FTaskTagScope Scope(ETaskTag::EGameThread);
+
 	GGameThreadId = FPlatformTLS::GetCurrentThreadId();
 
 	BootTimingPoint("android_main");
