Fix missing trace system shutdown when compiled without the engine causing missed event at end of application.

diff --git a/Engine/Source/Runtime/Launch/Private/LaunchEngineLoop.cpp b/Engine/Source/Runtime/Launch/Private/LaunchEngineLoop.cpp
--- a/Engine/Source/Runtime/Launch/Private/LaunchEngineLoop.cpp
+++ b/Engine/Source/Runtime/Launch/Private/LaunchEngineLoop.cpp
@@ -5885,6 +5885,11 @@ void FEngineLoop::AppExit( )
 	}
 
 	FInternationalization::TearDown();
+
+#if !WITH_ENGINE
+	// when compiled WITH_ENGINE, this will happen in FEngineLoop::Exit()
+	TRACE_CPUPROFILER_SHUTDOWN();
+#endif
 }
 
 void FEngineLoop::PostInitRHI()
