fixed stoping audio thread on exit: as AudioDeviceManager is released before audio thread stopped, this can lead to (subject of timing) an audio command executing after GAudioDeviceManager was set to nullptr, which causes null reference exception

diff --git a/Engine/Source/Runtime/Launch/Private/LaunchEngineLoop.cpp b/Engine/Source/Runtime/Launch/Private/LaunchEngineLoop.cpp
--- a/Engine/Source/Runtime/Launch/Private/LaunchEngineLoop.cpp
+++ b/Engine/Source/Runtime/Launch/Private/LaunchEngineLoop.cpp
@@ -4233,12 +4233,12 @@ void FEngineLoop::Exit()
 
 	if ( GEngine != nullptr )
 	{
-		GEngine->ReleaseAudioDeviceManager();
+		GEngine->PreExit();
 	}
 
-	if ( GEngine != nullptr )
+	if (GEngine != nullptr)
 	{
-		GEngine->PreExit();
+		GEngine->ReleaseAudioDeviceManager();
 	}
 
 	FAudioDeviceManager::Shutdown();
