IPackageResourceManager: Fix crashes in !WITH_ENGINE programs: startup and shutdown need to be called outside of WITH_ENGINE blocks.
#rnx

diff --git a/Engine/Source/Runtime/Launch/Private/LaunchEngineLoop.cpp b/Engine/Source/Runtime/Launch/Private/LaunchEngineLoop.cpp
--- a/Engine/Source/Runtime/Launch/Private/LaunchEngineLoop.cpp
+++ b/Engine/Source/Runtime/Launch/Private/LaunchEngineLoop.cpp
@@ -3302,6 +3302,11 @@ int32 FEngineLoop::PreInitPostStartupScreen(const TCHAR* CmdLine)
 			return 1;
 		}
 	}
+#else // !WITH_ENGINE
+#if WITH_COREUOBJECT
+	// Initialize the PackageResourceManager, which is needed to load any (non-script) Packages.
+	IPackageResourceManager::Initialize();
+#endif
 #endif // WITH_ENGINE
 
 #if WITH_COREUOBJECT
@@ -5847,6 +5852,12 @@ void FEngineLoop::AppPreExit( )
 		GetDerivedDataCacheRef().WaitForQuiescence(true);
 	}
 #endif
+#if !WITH_ENGINE
+#if WITH_COREUOBJECT
+	// Shutdown the PackageResourceManager in AppPreExit for programs that do not call FEngineLoop::Exit
+	IPackageResourceManager::Shutdown();
+#endif
+#endif
 
 #if WITH_EDITOR
 	FRemoteConfig::Flush();
