UE5 - Show actual path that couldn't be made
#rnx

diff --git a/Engine/Source/Runtime/Launch/Private/LaunchEngineLoop.cpp b/Engine/Source/Runtime/Launch/Private/LaunchEngineLoop.cpp
--- a/Engine/Source/Runtime/Launch/Private/LaunchEngineLoop.cpp
+++ b/Engine/Source/Runtime/Launch/Private/LaunchEngineLoop.cpp
@@ -1762,18 +1762,18 @@ int32 FEngineLoop::PreInitPreStartupScreen(const TCHAR* CmdLine)
 #if WITH_EDITOR
 	{
 		IPlatformFile& PlatformFile = FPlatformFileManager::Get().GetPlatformFile();
-
-		bool bCreateIntermediateSuccess = PlatformFile.CreateDirectory(*FPaths::ProjectIntermediateDir());
+		FString ProjectIntermediateDir = FPaths::ProjectIntermediateDir();
+		bool bCreateIntermediateSuccess = PlatformFile.CreateDirectory(*ProjectIntermediateDir);
 		if (!bCreateIntermediateSuccess)
 		{
-			UE_LOG(LogInit,Fatal,TEXT("Failed to create Intermediate directory."));
+			UE_LOG(LogInit,Fatal,TEXT("Failed to create Intermediate directory '%s'."), *ProjectIntermediateDir);
 		}
 
-		FString AutogenAbsolutePath = FPaths::ConvertRelativePathToFull(FPaths::ProjectIntermediateDir() / TEXT("ShaderAutogen"));
+		FString AutogenAbsolutePath = FPaths::ConvertRelativePathToFull(ProjectIntermediateDir / TEXT("ShaderAutogen"));
 		bool bCreateAutogenSuccess = PlatformFile.CreateDirectory(*AutogenAbsolutePath);
 		if (!bCreateAutogenSuccess)
 		{
-			UE_LOG(LogInit,Fatal,TEXT("Failed to create Intermediate/ShaderAutogen/ directory. Make sure Intermediate exists."));
+			UE_LOG(LogInit,Fatal,TEXT("Failed to create Intermediate/ShaderAutogen/ directory '%s'. Make sure Intermediate exists."), *AutogenAbsolutePath);
 		}
 
 		AddShaderSourceDirectoryMapping(TEXT("/ShaderAutogen"), AutogenAbsolutePath);
