Move up canonical setup of GIsEditor in the initialization of the engine. AppInit loads plugins which in turn make library function calls to things like the config system that want to know if it's an editor process running.

diff --git a/Engine/Source/Runtime/Launch/Private/LaunchEngineLoop.cpp b/Engine/Source/Runtime/Launch/Private/LaunchEngineLoop.cpp
--- a/Engine/Source/Runtime/Launch/Private/LaunchEngineLoop.cpp
+++ b/Engine/Source/Runtime/Launch/Private/LaunchEngineLoop.cpp
@@ -2261,7 +2261,34 @@ int32 FEngineLoop::PreInitPreStartupScreen(const TCHAR* CmdLine)
 
 #if WITH_ENGINE
 	AppLifetimeEventCapture::Init();
-#endif
+
+	// Perform canonical editor check before AppInit so plugin-consumed library code can know (they wouldn't otherwise)
+	// Things like the config system behave differently based on these globals, and we aren't the editor by default
+	if (bHasEditorToken)
+	{
+#if WITH_EDITOR
+
+		// We're the editor.
+		GIsClient = true;
+		GIsServer = true;
+		GIsEditor = true;
+		PRIVATE_GIsRunningCommandlet = false;
+
+		GWarn = &UnrealEdWarn;
+
+#else //WITH_EDITOR
+		FMessageDialog::Open(EAppMsgType::Ok, NSLOCTEXT("Engine", "EditorNotSupported", "Editor not supported in this mode."));
+		FPlatformMisc::RequestExit(false);
+		return 1;
+#endif //WITH_EDITOR
+	}
+#endif // WITH_ENGINE
+
+	// If we're not in the editor stop collecting the backlog now that we know
+	if (!GIsEditor)
+	{
+		GLog->EnableBacklog(false);
+	}
 
 #if WITH_ENGINE && TRACING_PROFILER
 PRAGMA_DISABLE_DEPRECATION_WARNINGS
@@ -2550,32 +2577,6 @@ PRAGMA_ENABLE_DEPRECATION_WARNINGS
 	GIsDemoMode = FParse::Param(FCommandLine::Get(), TEXT("DEMOMODE"));
 #endif
 
-	if (bHasEditorToken)
-	{
-#if WITH_EDITOR
-
-		// We're the editor.
-		GIsClient = true;
-		GIsServer = true;
-		GIsEditor = true;
-		PRIVATE_GIsRunningCommandlet = false;
-
-		GWarn = &UnrealEdWarn;
-
-#else
-		FMessageDialog::Open(EAppMsgType::Ok, NSLOCTEXT("Engine", "EditorNotSupported", "Editor not supported in this mode."));
-		FPlatformMisc::RequestExit(false);
-		return 1;
-#endif //WITH_EDITOR
-	}
-
-#endif // WITH_ENGINE
-	// If we're not in the editor stop collecting the backlog now that we know
-	if (!GIsEditor)
-	{
-		GLog->EnableBacklog(false);
-	}
-#if WITH_ENGINE
 
 	InitEngineTextLocalization();
 
