Fix for crash when using the MessageLog in Development builds.

diff --git a/Engine/Source/Runtime/Launch/Private/LaunchEngineLoop.cpp b/Engine/Source/Runtime/Launch/Private/LaunchEngineLoop.cpp
--- a/Engine/Source/Runtime/Launch/Private/LaunchEngineLoop.cpp
+++ b/Engine/Source/Runtime/Launch/Private/LaunchEngineLoop.cpp
@@ -3876,6 +3876,8 @@ bool FEngineLoop::LoadStartupCoreModules()
 	{
 #if WITH_UNREAL_DEVELOPER_TOOLS
 		FModuleManager::Get().LoadModule("MessageLog");
+		// Message Logs Developer implementation depends on the Editor style, it should be converted to core styles.
+		FModuleManager::Get().LoadModule("EditorStyle");
 #endif	// WITH_UNREAL_DEVELOPER_TOOLS
 #if WITH_EDITOR
 		FModuleManager::Get().LoadModule("CollisionAnalyzer");
