DDC: Always initialize DDC in scenarios that do not require cooked content

This avoids a crash when the first use of DDC is from a worker thread when global shader map compilation is skipped due to use of the null RHI.

diff --git a/Engine/Source/Runtime/Launch/Private/LaunchEngineLoop.cpp b/Engine/Source/Runtime/Launch/Private/LaunchEngineLoop.cpp
--- a/Engine/Source/Runtime/Launch/Private/LaunchEngineLoop.cpp
+++ b/Engine/Source/Runtime/Launch/Private/LaunchEngineLoop.cpp
@@ -2699,6 +2699,9 @@ PRAGMA_ENABLE_DEPRECATION_WARNINGS
 
 	if (!FPlatformProperties::RequiresCookedData())
 	{
+		// Ensure that DDC is initialized from the game thread.
+		GetDerivedDataCacheRef();
+
 		check(!GDistanceFieldAsyncQueue);
 		GDistanceFieldAsyncQueue = new FDistanceFieldAsyncQueue();
 
@@ -2753,10 +2756,6 @@ PRAGMA_ENABLE_DEPRECATION_WARNINGS
 				return 1;
 			}
 		}
-		else if (FPlatformProperties::RequiresCookedData() == false)
-		{
-			GetDerivedDataCacheRef();
-		}
 
 		{
 			SCOPED_BOOT_TIMING("CreateMoviePlayer");
