commented out usage of `PoolThreadStackSize` and `TaskThreadStackSize` config ini parameters as the config is not initialized yet when they are used.

diff --git a/Engine/Source/Runtime/Core/Private/Async/TaskGraph.cpp b/Engine/Source/Runtime/Core/Private/Async/TaskGraph.cpp
--- a/Engine/Source/Runtime/Core/Private/Async/TaskGraph.cpp
+++ b/Engine/Source/Runtime/Core/Private/Async/TaskGraph.cpp
@@ -1323,7 +1323,8 @@ public:
 			StackSize = 384 * 1024;
 #endif
 
-			GConfig->GetInt(TEXT("Core.System"), TEXT("TaskThreadStackSize"), StackSize, GEngineIni);
+			// GConfig is not initialized yet, the only solution for now is to hardcode desired values
+			//GConfig->GetInt(TEXT("Core.System"), TEXT("TaskThreadStackSize"), StackSize, GEngineIni);
 
 			if (GroupName != PrevGroupName)
 			{
diff --git a/Engine/Source/Runtime/Launch/Private/LaunchEngineLoop.cpp b/Engine/Source/Runtime/Launch/Private/LaunchEngineLoop.cpp
--- a/Engine/Source/Runtime/Launch/Private/LaunchEngineLoop.cpp
+++ b/Engine/Source/Runtime/Launch/Private/LaunchEngineLoop.cpp
@@ -2170,7 +2170,8 @@ int32 FEngineLoop::PreInitPreStartupScreen(const TCHAR* CmdLine)
 		SCOPED_BOOT_TIMING("Init FQueuedThreadPool's");
 
 		int32 StackSize = 128 * 1024;
-		GConfig->GetInt(TEXT("Core.System"), TEXT("PoolThreadStackSize"), StackSize, GEngineIni);
+		// GConfig is not initialized yet, the only solution for now is to hardcode desired values
+		// GConfig->GetInt(TEXT("Core.System"), TEXT("PoolThreadStackSize"), StackSize, GEngineIni);
 
 		bool bForceEditorStackSize = false;
 #if WITH_EDITOR
