Move stall detector shutdown to AppExit to fix some late-constructed detectors causing assert at exit-time.

diff --git a/Engine/Source/Runtime/Launch/Private/LaunchEngineLoop.cpp b/Engine/Source/Runtime/Launch/Private/LaunchEngineLoop.cpp
--- a/Engine/Source/Runtime/Launch/Private/LaunchEngineLoop.cpp
+++ b/Engine/Source/Runtime/Launch/Private/LaunchEngineLoop.cpp
@@ -5942,10 +5942,6 @@ void FEngineLoop::AppPreExit( )
 	FRemoteConfig::Flush();
 #endif
 
-#if STALL_DETECTOR
-	UE::FStallDetector::Shutdown();
-#endif
-
 	FCoreDelegates::OnExit.Broadcast();
 
 	// Clean up the thread pool
@@ -6002,8 +5998,12 @@ void FEngineLoop::AppPreExit( )
 }
 
 
-void FEngineLoop::AppExit( )
+void FEngineLoop::AppExit()
 {
+#if STALL_DETECTOR
+	UE::FStallDetector::Shutdown();
+#endif
+
 #if !WITH_ENGINE
 	// when compiled WITH_ENGINE, this will happen in FEngineLoop::Exit()
 	FTaskGraphInterface::Shutdown();
