Tagging main thread as game thread on Lumin (similar to CL 14630758 for Android)

diff --git a/Engine/Source/Runtime/Launch/Private/Lumin/LaunchLumin.cpp b/Engine/Source/Runtime/Launch/Private/Lumin/LaunchLumin.cpp
--- a/Engine/Source/Runtime/Launch/Private/Lumin/LaunchLumin.cpp
+++ b/Engine/Source/Runtime/Launch/Private/Lumin/LaunchLumin.cpp
@@ -158,6 +158,7 @@ static bool IncreasePerProcessLimits()
 
 int main(int argc, char *argv[])
 {
+	FTaskTagScope Scope(ETaskTag::EGameThread);
 	FPlatformMisc::SetGracefulTerminationHandler();
 
 	int ErrorLevel = 0;
