Move the creation and mapping of shader directories to later on in the start up sequence.

- Creating the shader directories relies on the process having a valid project but was being called before we validate the project file.
- This meant that if the project path is invalid instead of getting a clear error dialog about it the user would hit a fatal log while generating the shader directories instead.
- We can safely move the code to a point after the project path has been validated as nothing will have tried to access the list of shader directories yet.

#rnx

diff --git a/Engine/Source/Runtime/Launch/Private/LaunchEngineLoop.cpp b/Engine/Source/Runtime/Launch/Private/LaunchEngineLoop.cpp
--- a/Engine/Source/Runtime/Launch/Private/LaunchEngineLoop.cpp
+++ b/Engine/Source/Runtime/Launch/Private/LaunchEngineLoop.cpp
@@ -1771,29 +1771,6 @@ int32 FEngineLoop::PreInitPreStartupScreen(const TCHAR* CmdLine)
 	FString Token = FParse::Token(ParsedCmdLine, 0);
 
 #if WITH_ENGINE
-	// Add the default engine shader dir
-	AddShaderSourceDirectoryMapping(TEXT("/Engine"), FGenericPlatformProcess::ShaderDir());
-
-#if WITH_EDITOR
-	{
-		IPlatformFile& PlatformFile = FPlatformFileManager::Get().GetPlatformFile();
-		FString ProjectIntermediateDir = FPaths::ProjectIntermediateDir();
-		bool bCreateIntermediateSuccess = PlatformFile.CreateDirectory(*ProjectIntermediateDir);
-		if (!bCreateIntermediateSuccess)
-		{
-			UE_LOG(LogInit,Fatal,TEXT("Failed to create Intermediate directory '%s'."), *ProjectIntermediateDir);
-		}
-
-		FString AutogenAbsolutePath = FPaths::ConvertRelativePathToFull(ProjectIntermediateDir / TEXT("ShaderAutogen"));
-		bool bCreateAutogenSuccess = PlatformFile.CreateDirectory(*AutogenAbsolutePath);
-		if (!bCreateAutogenSuccess)
-		{
-			UE_LOG(LogInit,Fatal,TEXT("Failed to create Intermediate/ShaderAutogen/ directory '%s'. Make sure Intermediate exists."), *AutogenAbsolutePath);
-		}
-
-		AddShaderSourceDirectoryMapping(TEXT("/ShaderAutogen"), AutogenAbsolutePath);
-	}
-#endif
 
 	TArray<FString> Tokens;
 	TArray<FString> Switches;
@@ -2077,6 +2054,32 @@ int32 FEngineLoop::PreInitPreStartupScreen(const TCHAR* CmdLine)
 	}
 #endif
 
+#if WITH_ENGINE
+	// Add the default engine shader dir
+	AddShaderSourceDirectoryMapping(TEXT("/Engine"), FGenericPlatformProcess::ShaderDir());
+
+#if WITH_EDITOR
+	{
+		IPlatformFile& PlatformFile = FPlatformFileManager::Get().GetPlatformFile();
+		FString ProjectIntermediateDir = FPaths::ProjectIntermediateDir();
+		bool bCreateIntermediateSuccess = PlatformFile.CreateDirectory(*ProjectIntermediateDir);
+		if (!bCreateIntermediateSuccess)
+		{
+			UE_LOG(LogInit, Fatal, TEXT("Failed to create Intermediate directory '%s'."), *ProjectIntermediateDir);
+		}
+
+		FString AutogenAbsolutePath = FPaths::ConvertRelativePathToFull(ProjectIntermediateDir / TEXT("ShaderAutogen"));
+		bool bCreateAutogenSuccess = PlatformFile.CreateDirectory(*AutogenAbsolutePath);
+		if (!bCreateAutogenSuccess)
+		{
+			UE_LOG(LogInit, Fatal, TEXT("Failed to create Intermediate/ShaderAutogen/ directory '%s'. Make sure Intermediate exists."), *AutogenAbsolutePath);
+		}
+
+		AddShaderSourceDirectoryMapping(TEXT("/ShaderAutogen"), AutogenAbsolutePath);
+	}
+#endif //WITH_EDITOR
+#endif //WITH_ENGINE
+
 	// Some programs might not use the taskgraph or thread pool
 	bool bCreateTaskGraphAndThreadPools = true;
 	// If STATS is defined (via FORCE_USE_STATS or other), we have to call FTaskGraphInterface::Startup()
