Adding "Compiling Global Shaders..." progress frame to track what's happening on the editor's splash screen.

- The Expected Work This Frame  25 value was arbitrarily picked.

diff --git a/Engine/Source/Runtime/Launch/Private/LaunchEngineLoop.cpp b/Engine/Source/Runtime/Launch/Private/LaunchEngineLoop.cpp
--- a/Engine/Source/Runtime/Launch/Private/LaunchEngineLoop.cpp
+++ b/Engine/Source/Runtime/Launch/Private/LaunchEngineLoop.cpp
@@ -2759,7 +2759,7 @@ PRAGMA_ENABLE_DEPRECATION_WARNINGS
 
 		FDelayedAutoRegisterHelper::RunAndClearDelayedAutoRegisterDelegates(EDelayedRegisterRunPhase::ShaderTypesReady);
 
-		SlowTask.EnterProgressFrame(30);
+		SlowTask.EnterProgressFrame(25, LOCTEXT("CompileGlobalShaderMap", "Compiling Global Shaders..."));
 
 		// Load the global shaders.
 		// if (!IsRunningCommandlet())
@@ -2779,6 +2779,8 @@ PRAGMA_ENABLE_DEPRECATION_WARNINGS
 			}
 		}
 
+		SlowTask.EnterProgressFrame(5);
+
 		{
 			SCOPED_BOOT_TIMING("CreateMoviePlayer");
 			CreateMoviePlayer();
