Remove UE3/UE4 references

diff --git a/Engine/Build/Android/Java/src/com/epicgames/unreal/GameActivity.java.template b/Engine/Build/Android/Java/src/com/epicgames/unreal/GameActivity.java.template
--- a/Engine/Build/Android/Java/src/com/epicgames/unreal/GameActivity.java.template
+++ b/Engine/Build/Android/Java/src/com/epicgames/unreal/GameActivity.java.template
@@ -5553,7 +5553,7 @@ public class GameActivity extends $${gameActivitySuperClass}$$ implements Surfac
 
 	public int AndroidThunkJava_GetMetaDataInt(String key)
 	{
-		if (key.equals("ue4.http.proxy.proxyPort"))
+		if (key.equals("unreal.http.proxy.proxyPort"))
 		{
 			if (ANDROID_BUILD_VERSION >= 14)
 			{
@@ -5576,7 +5576,7 @@ public class GameActivity extends $${gameActivitySuperClass}$$ implements Surfac
 			return VulkanLevel;
 		}
 		else
-		if (key.equals("ue4.getUsedMemory"))
+		if (key.equals("unreal.getUsedMemory"))
 		{
 			int ProcessMemory = 0;
 			synchronized(_activity)
@@ -5614,7 +5614,7 @@ public class GameActivity extends $${gameActivitySuperClass}$$ implements Surfac
 
 	public long AndroidThunkJava_GetMetaDataLong(String key)
 	{
-		if (key.equals("ue4.display.PresentationDeadlineNanos"))
+		if (key.equals("unreal.display.PresentationDeadlineNanos"))
 		{
 			if (android.os.Build.VERSION.SDK_INT >= 21)
 			{
@@ -5625,7 +5625,7 @@ public class GameActivity extends $${gameActivitySuperClass}$$ implements Surfac
 			return -1;
 		}
 		else
-		if (key.equals("ue4.display.AppVsyncOffsetNanos"))
+		if (key.equals("unreal.display.AppVsyncOffsetNanos"))
 		{
 			if (android.os.Build.VERSION.SDK_INT >= 21)
 			{
@@ -5640,7 +5640,7 @@ public class GameActivity extends $${gameActivitySuperClass}$$ implements Surfac
 
 	public float AndroidThunkJava_GetMetaDataFloat(String key)
 	{
-		if (key.equals("ue4.display.getRefreshRate"))
+		if (key.equals("unreal.display.getRefreshRate"))
 		{
 			WindowManager windowManager = (WindowManager)getSystemService(WINDOW_SERVICE);
 			android.view.Display display = windowManager.getDefaultDisplay();
@@ -5658,7 +5658,7 @@ public class GameActivity extends $${gameActivitySuperClass}$$ implements Surfac
 			return getProp(key.substring(8));
 		}
 		else
-		if (key.equals("ue4.http.proxy.proxyHost"))
+		if (key.equals("unreal.http.proxy.proxyHost"))
 		{
 			if (ANDROID_BUILD_VERSION >= 14)
 			{
@@ -5670,7 +5670,7 @@ public class GameActivity extends $${gameActivitySuperClass}$$ implements Surfac
 			}
 		}
 		else
-		if (key.equals("ue4.displaymetrics.dpi"))
+		if (key.equals("unreal.displaymetrics.dpi"))
 		{
 			DisplayMetrics metrics = new DisplayMetrics();
 			if (android.os.Build.VERSION.SDK_INT >= 17)
diff --git a/Engine/Build/Android/UE4Game/PushCommandLine.bat b/Engine/Build/Android/UE4Game/PushCommandLine.bat
deleted file mode 100755
diff --git a/Engine/Build/Android/UnrealGame/PushCommandLine.bat b/Engine/Build/Android/UnrealGame/PushCommandLine.bat
new file mode 100755
--- /dev/null
+++ b/Engine/Build/Android/UnrealGame/PushCommandLine.bat
@@ -0,0 +1 @@
+%ANDROID_HOME%\platform-tools\adb.exe push UE4CommandLine.txt /mnt/sdcard/UnrealGame/UnrealGame/UE4CommandLine.txt
\ No newline at end of file
diff --git a/Engine/Build/IOS/Resources/FrameworkWrapper/FrameworkWrapper/IOS/AppDelegate.h b/Engine/Build/IOS/Resources/FrameworkWrapper/FrameworkWrapper/IOS/AppDelegate.h
--- a/Engine/Build/IOS/Resources/FrameworkWrapper/FrameworkWrapper/IOS/AppDelegate.h
+++ b/Engine/Build/IOS/Resources/FrameworkWrapper/FrameworkWrapper/IOS/AppDelegate.h
@@ -10,7 +10,7 @@
 
 #include "UnrealView.h"
 
-#if CAN_USE_UE4 && __cplusplus
+#if CAN_USE_UE && __cplusplus
     #pragma clang diagnostic push
     #pragma clang diagnostic ignored "-Weverything"
 
@@ -25,7 +25,7 @@
         - (BOOL)application:(UIApplication *)application didFinishLaunchingWithOptions:(NSDictionary *)launchOptions;
         - (void)applicationDidEnterBackground:(UIApplication *)application;
     @end
-#endif // CAN_USE_UE4
+#endif // CAN_USE_UE
 
 
 @interface AppDelegate : IOSAppDelegate <UIApplicationDelegate>
diff --git a/Engine/Build/IOS/Resources/FrameworkWrapper/FrameworkWrapper/IOS/AppDelegate.mm b/Engine/Build/IOS/Resources/FrameworkWrapper/FrameworkWrapper/IOS/AppDelegate.mm
--- a/Engine/Build/IOS/Resources/FrameworkWrapper/FrameworkWrapper/IOS/AppDelegate.mm
+++ b/Engine/Build/IOS/Resources/FrameworkWrapper/FrameworkWrapper/IOS/AppDelegate.mm
@@ -9,14 +9,14 @@
 #import "AppDelegate.h"
 #import "UnrealView.h"
 
-#if !CAN_USE_UE4 || !__cplusplus
+#if !CAN_USE_UE || !__cplusplus
     // Stub class to inherit from in a non-UE context
     @implementation IOSAppDelegate
         - (void)applicationWillEnterForeground:(UIApplication *)application {}
         - (BOOL)application:(UIApplication *)application didFinishLaunchingWithOptions:(NSDictionary *)launchOptions { return YES; }
         - (void)applicationDidEnterBackground:(UIApplication *)application {}
     @end
-#endif // !CAN_USE_UE4
+#endif // !CAN_USE_UE
 
 @implementation AppDelegate
 
diff --git a/Engine/Build/IOS/Resources/FrameworkWrapper/FrameworkWrapper/Views/UnrealView.h b/Engine/Build/IOS/Resources/FrameworkWrapper/FrameworkWrapper/Views/UnrealView.h
--- a/Engine/Build/IOS/Resources/FrameworkWrapper/FrameworkWrapper/Views/UnrealView.h
+++ b/Engine/Build/IOS/Resources/FrameworkWrapper/FrameworkWrapper/Views/UnrealView.h
@@ -7,9 +7,9 @@
 //
 
 #if defined(__aarch64__)
-    #define CAN_USE_UE4 1
+    #define CAN_USE_UE 1
 #else
-    #define CAN_USE_UE4 0
+    #define CAN_USE_UE 0
 #endif
 
 #define NW_LOG NSLog
diff --git a/Engine/Build/IOS/Resources/FrameworkWrapper/FrameworkWrapper/Views/UnrealView.mm b/Engine/Build/IOS/Resources/FrameworkWrapper/FrameworkWrapper/Views/UnrealView.mm
--- a/Engine/Build/IOS/Resources/FrameworkWrapper/FrameworkWrapper/Views/UnrealView.mm
+++ b/Engine/Build/IOS/Resources/FrameworkWrapper/FrameworkWrapper/Views/UnrealView.mm
@@ -9,7 +9,7 @@
 
 #import "UnrealView.h"
 
-#if CAN_USE_UE4 && __cplusplus
+#if CAN_USE_UE && __cplusplus
     #pragma clang diagnostic push
     #pragma clang diagnostic ignored "-Weverything"
 
@@ -18,7 +18,7 @@
     #include "Misc/EmbeddedCommunication.h"
 
     #pragma clang diagnostic pop
-#endif // CAN_USE_UE4
+#endif // CAN_USE_UE
 
 
 @implementation UnrealContainerView
@@ -34,7 +34,7 @@ -(void) awakeFromNib
 
 -(BOOL) IsUnrealUsable
 {
-#if CAN_USE_UE4
+#if CAN_USE_UE
 	// check memory amount, and running OS version, etc
 	return YES;
 #else
@@ -55,7 +55,7 @@ +(void)DelayedCreateView
     NW_LOG(@"DelayedCreateView start");
 
 	// @todo these views will never be released
-#if CAN_USE_UE4
+#if CAN_USE_UE
 	if ([DelayedCreateContainer IsUnrealUsable])
 	{
 		sharedUnrealView = [[FIOSView alloc] initWithFrame:DelayedCreateContainer.bounds];
@@ -88,7 +88,7 @@ -(void) createUnrealView
 {
     if (sharedUnrealView == nil)
 	{
-#if CAN_USE_UE4
+#if CAN_USE_UE
 		
 //       [self SetupEmbeddedToNativeCallback];
 
@@ -112,7 +112,7 @@ -(void) createUnrealView
 }
 
 
-#if CAN_USE_UE4
+#if CAN_USE_UE
 +(id)UE4ToGeneric:(const FEmbeddedCommunicationMap&)UE4Map
 {
 	// look for a special single entry { "json", "xxxx" } map
@@ -147,7 +147,7 @@ +(void)EmbeddedCallForSubsystem:(NSString* _Nullable)InSubsystem WithCommand:(NS
 	// we have to copy the handler because it will go out of scope after this call (if it's a stack block) and can't be retained or anything
 	EmbeddedParamsCompletionType Handler = [InHandler copy];
 	
-#if CAN_USE_UE4
+#if CAN_USE_UE
 	FEmbeddedCallParamsHelper CallHelper;
 	__block FEmbeddedCallParamsHelper* BlockHelper = &CallHelper;
 	CallHelper.Command = Command;
@@ -279,14 +279,14 @@ +(void)JSBridgeCallByName:(NSString*)UObjectName MethodName:(NSString*)MethodNam
 
 +(void)SetEmbeddedObject:(void* _Nullable)Object ForName:(NSString* _Nonnull)Name
 {
-#if CAN_USE_UE4
+#if CAN_USE_UE
 	FEmbeddedDelegates::SetNamedObject(Name, Object);
 #endif
 }
 
 +(void)HandleTouchMessage:(NSDictionary*)Message
 {
-#if CAN_USE_UE4
+#if CAN_USE_UE
 
 	dispatch_async(dispatch_get_main_queue(), ^
 	   {
@@ -324,7 +324,7 @@ +(void)HandleTouchMessage:(NSDictionary*)Message
 
 +(void)WakeUpUnreal
 {
-#if CAN_USE_UE4
+#if CAN_USE_UE
 	FEmbeddedCommunication::KeepAwake(TEXT("Native"), true);
 #endif
 
@@ -334,7 +334,7 @@ +(void)WakeUpUnreal
 
 +(void)AllowUnrealToSleep
 {
-#if CAN_USE_UE4
+#if CAN_USE_UE
 	FEmbeddedCommunication::AllowSleep(TEXT("Native"));
 #endif
 
diff --git a/Engine/Build/IOS/UE4CmdLineRun/UE4CmdLineRun.m b/Engine/Build/IOS/UE4CmdLineRun/UE4CmdLineRun.m
deleted file mode 100644
diff --git a/Engine/Build/IOS/UECmdLineRun/UECmdLineRun.m b/Engine/Build/IOS/UECmdLineRun/UECmdLineRun.m
new file mode 100644
--- /dev/null
+++ b/Engine/Build/IOS/UECmdLineRun/UECmdLineRun.m
@@ -0,0 +1,33 @@
+// Copyright Epic Games, Inc. All Rights Reserved.
+
+#import <UIKit/UIKit.h>
+#import <Foundation/Foundation.h>
+#import <XCTest/XCTest.h>
+
+@interface UE4CmdLineRun : XCTestCase
+
+@end
+
+@implementation UE4CmdLineRun
+
+- (void)setUp
+{
+    [super setUp];
+    // Put setup code here. This method is called before the invocation of each test method in the class.
+}
+
+- (void)tearDown
+{
+    // Put teardown code here. This method is called after the invocation of each test method in the class.
+    [super tearDown];
+}
+
+- (void)testCmdLineRun
+{
+	// by running the runloop from the test, the full game will be running as if launched normalled
+	// the unit test is run from the startup of the normal application path, this lets it continue on
+	// since the system will quit the app as soon as all tests are complete
+	[[NSRunLoop currentRunLoop] run];
+}
+
+@end
diff --git a/Engine/Plugins/Runtime/MicrosoftSpatialAudio/Source/MicrosoftSpatialSound/Private/MicrosoftSpatialSoundPlugin.cpp b/Engine/Plugins/Runtime/MicrosoftSpatialAudio/Source/MicrosoftSpatialSound/Private/MicrosoftSpatialSoundPlugin.cpp
--- a/Engine/Plugins/Runtime/MicrosoftSpatialAudio/Source/MicrosoftSpatialSound/Private/MicrosoftSpatialSoundPlugin.cpp
+++ b/Engine/Plugins/Runtime/MicrosoftSpatialAudio/Source/MicrosoftSpatialSound/Private/MicrosoftSpatialSoundPlugin.cpp
@@ -326,7 +326,7 @@ uint32 FMicrosoftSpatialSound::Run()
 	{
 		if (!bWarnedMicrosoftSpatialSoundDynamicObjectCountIsZero && SAC->GetMaxDynamicObjects() == 0)
 		{
-			UE_LOG(LogMicrosoftSpatialSound, Warning, TEXT("Microsoft Spatial Sound has zero MaxDynamicObjects.  No sounds can play!  You need to enable Spatial Sound (Windows Sonic for Headphones) in your PC audio settings then restart UE4."));
+			UE_LOG(LogMicrosoftSpatialSound, Warning, TEXT("Microsoft Spatial Sound has zero MaxDynamicObjects.  No sounds can play!  You need to enable Spatial Sound (Windows Sonic for Headphones) in your PC audio settings then restart Unreal."));
 			bWarnedMicrosoftSpatialSoundDynamicObjectCountIsZero = true;
 		}
 
@@ -413,7 +413,7 @@ void FMicrosoftSpatialSoundModule::StartupModule()
 	{
 		UE_LOG(LogMicrosoftSpatialSound,
 			Warning,
-			TEXT("Microsoft Spatial Sound for UE4 currently only supports windows version '%d' or higher (Current version: '%d')"),
+			TEXT("Microsoft Spatial Sound for Unreal currently only supports windows version '%d' or higher (Current version: '%d')"),
 			MIN_WIN_10_VERSION_FOR_WMR_SPATSOUND,
 			CurrentVersionNumber);
 
diff --git a/Engine/Source/Programs/ChaosVisualDebugger/Private/Mac/ChaosVisualDebuggerMac.cpp b/Engine/Source/Programs/ChaosVisualDebugger/Private/Mac/ChaosVisualDebuggerMac.cpp
--- a/Engine/Source/Programs/ChaosVisualDebugger/Private/Mac/ChaosVisualDebuggerMac.cpp
+++ b/Engine/Source/Programs/ChaosVisualDebugger/Private/Mac/ChaosVisualDebuggerMac.cpp
@@ -10,14 +10,14 @@
 static FString GSavedCommandLine;
 
 
-@interface UE4AppDelegate : NSObject<NSApplicationDelegate, NSFileManagerDelegate>
+@interface UEAppDelegate : NSObject<NSApplicationDelegate, NSFileManagerDelegate>
 {
 }
 
 @end
 
 
-@implementation UE4AppDelegate
+@implementation UEAppDelegate
 
 //handler for the quit apple event used by the Dock menu
 - (void)handleQuitEvent:(NSAppleEventDescriptor*)Event withReplyEvent:(NSAppleEventDescriptor*)ReplyEvent
@@ -112,7 +112,7 @@ int main(int argc, char *argv[])
 
 	SCOPED_AUTORELEASE_POOL;
 	[NSApplication sharedApplication];
-	[NSApp setDelegate:[UE4AppDelegate new]];
+	[NSApp setDelegate:[UEAppDelegate new]];
 	[NSApp run];
 	return 0;
 }
diff --git a/Engine/Source/Programs/CrashReportClient/Private/Mac/CrashReportClientMainMac.cpp b/Engine/Source/Programs/CrashReportClient/Private/Mac/CrashReportClientMainMac.cpp
--- a/Engine/Source/Programs/CrashReportClient/Private/Mac/CrashReportClientMainMac.cpp
+++ b/Engine/Source/Programs/CrashReportClient/Private/Mac/CrashReportClientMainMac.cpp
@@ -20,13 +20,13 @@ void CrashReporterCrashHandler(const FGenericCrashContext& GenericContext)
 static FString GSavedCommandLine;
 static bool GIsUnattended = false;
 
-@interface UE4AppDelegate : NSObject <NSApplicationDelegate, NSFileManagerDelegate>
+@interface UEAppDelegate : NSObject <NSApplicationDelegate, NSFileManagerDelegate>
 {
 }
 
 @end
 
-@implementation UE4AppDelegate
+@implementation UEAppDelegate
 
 //handler for the quit apple event used by the Dock menu
 - (void)handleQuitEvent:(NSAppleEventDescriptor*)Event withReplyEvent:(NSAppleEventDescriptor*)ReplyEvent
@@ -110,7 +110,7 @@ int main(int argc, char *argv[])
 	
 	SCOPED_AUTORELEASE_POOL;
 	[NSApplication sharedApplication];
-	[NSApp setDelegate:[UE4AppDelegate new]];
+	[NSApp setDelegate:[UEAppDelegate new]];
 	[NSApp run];
 	return 0;
 }
diff --git a/Engine/Source/Programs/IOS/iPhonePackager/CodeSigning.cs b/Engine/Source/Programs/IOS/iPhonePackager/CodeSigning.cs
--- a/Engine/Source/Programs/IOS/iPhonePackager/CodeSigning.cs
+++ b/Engine/Source/Programs/IOS/iPhonePackager/CodeSigning.cs
@@ -738,7 +738,7 @@ namespace iPhonePackager
 		/// Does the actual work of signing the application
 		///   Modifies the following files:
 		///	 Info.plist
-		///	 [Executable] (file name derived from CFBundleExecutable in the Info.plist, e.g., UDKGame)
+		///	 [Executable] (file name derived from CFBundleExecutable in the Info.plist, e.g., UnrealGame)
 		///	 _CodeSignature/CodeResources
 		///	 [ResourceRules] (file name derived from CFBundleResourceSpecification, e.g., CustomResourceRules.plist)
 		/// </summary>
diff --git a/Engine/Source/Programs/IOS/iPhonePackager/Config.cs b/Engine/Source/Programs/IOS/iPhonePackager/Config.cs
--- a/Engine/Source/Programs/IOS/iPhonePackager/Config.cs
+++ b/Engine/Source/Programs/IOS/iPhonePackager/Config.cs
@@ -229,7 +229,7 @@ namespace iPhonePackager
 		/// </summary>
 		public static string GetIPAPath(string FileSuffix)
 		{
-			// Quash the default Epic_ so that stubs for UDK installers get named correctly and can be used
+			// Quash the default Epic_ so that stubs for Launcher installers get named correctly and can be used
 			string FilePrefix = (SigningPrefix == "Epic_") ? "" : SigningPrefix;
 
 			string Filename;
@@ -248,7 +248,7 @@ namespace iPhonePackager
 
 		public static string RemapIPAPath(string FileSuffix)
 		{
-			// Quash the default Epic_ so that stubs for UDK installers get named correctly and can be used
+			// Quash the default Epic_ so that stubs for Launcher installers get named correctly and can be used
 			string FilePrefix = (SigningPrefix == "Epic_") ? "" : SigningPrefix;
 
 			string Filename;
diff --git a/Engine/Source/Programs/SlateViewer/Private/Mac/SlateViewerMainMac.cpp b/Engine/Source/Programs/SlateViewer/Private/Mac/SlateViewerMainMac.cpp
--- a/Engine/Source/Programs/SlateViewer/Private/Mac/SlateViewerMainMac.cpp
+++ b/Engine/Source/Programs/SlateViewer/Private/Mac/SlateViewerMainMac.cpp
@@ -6,13 +6,13 @@
 
 static FString GSavedCommandLine;
 
-@interface UE4AppDelegate : NSObject <NSApplicationDelegate, NSFileManagerDelegate>
+@interface UEAppDelegate : NSObject <NSApplicationDelegate, NSFileManagerDelegate>
 {
 }
 
 @end
 
-@implementation UE4AppDelegate
+@implementation UEAppDelegate
 
 //handler for the quit apple event used by the Dock menu
 - (void)handleQuitEvent:(NSAppleEventDescriptor*)Event withReplyEvent:(NSAppleEventDescriptor*)ReplyEvent
@@ -79,7 +79,7 @@ int main(int argc, char *argv[])
 
 	SCOPED_AUTORELEASE_POOL;
 	[NSApplication sharedApplication];
-	[NSApp setDelegate:[UE4AppDelegate new]];
+	[NSApp setDelegate:[UEAppDelegate new]];
 	[NSApp run];
 	return 0;
 }
diff --git a/Engine/Source/Programs/UnrealBuildTool/System/UnrealPluginLanguage.cs b/Engine/Source/Programs/UnrealBuildTool/System/UnrealPluginLanguage.cs
--- a/Engine/Source/Programs/UnrealBuildTool/System/UnrealPluginLanguage.cs
+++ b/Engine/Source/Programs/UnrealBuildTool/System/UnrealPluginLanguage.cs
@@ -415,7 +415,7 @@ namespace UnrealBuildTool
 	 * 	<!-- optional additions to GameActivity onActivityResult in GameActivity.java -->
 	 * 	<gameActivityOnActivityResultAdditions>	</gameActivityOnActivityResultAdditions>
 	 * 	
-	 * 	<!-- optional libraries to load in GameActivity.java before libUE4.so -->
+	 * 	<!-- optional libraries to load in GameActivity.java before libUnreal.so -->
 	 * 	<soLoadLibrary>	</soLoadLibrary>
 	 * 	
 	 * 
@@ -3306,13 +3306,13 @@ namespace UnrealBuildTool
 
 			if(UserSettingsFolder != null)
 			{
-				// <AppData>/UE4/EngineConfig/User* ini
+				// <AppData>/Unreal/EngineConfig/User* ini
 				yield return FileReference.Combine(UserSettingsFolder, "Unreal Engine", "Engine", "Config", "User" + BaseIniName + ".ini");
 			}
 
 			if(PersonalFolder != null)
 			{
-				// <Documents>/UE4/EngineConfig/User* ini
+				// <Documents>/Unreal/EngineConfig/User* ini
 				yield return FileReference.Combine(PersonalFolder, "Unreal Engine", "Engine", "Config", "User" + BaseIniName + ".ini");
 			}
 
diff --git a/Engine/Source/Programs/UnrealMultiUserServer/Private/UnrealMultiUserServerMain.cpp b/Engine/Source/Programs/UnrealMultiUserServer/Private/UnrealMultiUserServerMain.cpp
--- a/Engine/Source/Programs/UnrealMultiUserServer/Private/UnrealMultiUserServerMain.cpp
+++ b/Engine/Source/Programs/UnrealMultiUserServer/Private/UnrealMultiUserServerMain.cpp
@@ -48,10 +48,10 @@ public:
 
 static CommandLineArguments GSavedCommandLine;
 
-@interface UE4AppDelegate : NSObject <NSApplicationDelegate>
+@interface UEAppDelegate : NSObject <NSApplicationDelegate>
 @end
 
-@implementation UE4AppDelegate
+@implementation UEAppDelegate
 
 //handler for the quit apple event used by the Dock menu
 - (void)handleQuitEvent:(NSAppleEventDescriptor*)Event withReplyEvent:(NSAppleEventDescriptor*)ReplyEvent
@@ -138,7 +138,7 @@ int main(int argc, char *argv[])
 	// Launch the application.
 	SCOPED_AUTORELEASE_POOL;
 	[NSApplication sharedApplication];
-	[NSApp setDelegate:[UE4AppDelegate new]];
+	[NSApp setDelegate:[UEAppDelegate new]];
 	[NSApp run];
 	return 0;
 }
diff --git a/Engine/Source/Runtime/ApplicationCore/Private/Android/AndroidInputInterface.cpp b/Engine/Source/Runtime/ApplicationCore/Private/Android/AndroidInputInterface.cpp
--- a/Engine/Source/Runtime/ApplicationCore/Private/Android/AndroidInputInterface.cpp
+++ b/Engine/Source/Runtime/ApplicationCore/Private/Android/AndroidInputInterface.cpp
@@ -1604,9 +1604,9 @@ void FAndroidInputInterface::MouseButtonEvent(int32 deviceId, int32 buttonId, bo
 	FScopeLock Lock(&TouchInputCriticalSection);
 
 	MouseEventType EventType = buttonDown ? MouseEventType::MouseButtonDown : MouseEventType::MouseButtonUp;
-	EMouseButtons::Type UE4Button = (buttonId == 0) ? EMouseButtons::Left : (buttonId == 1) ? EMouseButtons::Right : EMouseButtons::Middle;
+	EMouseButtons::Type UnrealButton = (buttonId == 0) ? EMouseButtons::Left : (buttonId == 1) ? EMouseButtons::Right : EMouseButtons::Middle;
 	FAndroidInputInterface::MouseDataStack.Push(
-		MouseData{ EventType, UE4Button, 0, 0, 0, 0, 0.0f });
+		MouseData{ EventType, UnrealButton, 0, 0, 0, 0, 0.0f });
 }
 
 void FAndroidInputInterface::DeferMessage(const FDeferredAndroidMessage& DeferredMessage)
diff --git a/Engine/Source/Runtime/ApplicationCore/Private/Android/AndroidPlatformApplicationMisc.cpp b/Engine/Source/Runtime/ApplicationCore/Private/Android/AndroidPlatformApplicationMisc.cpp
--- a/Engine/Source/Runtime/ApplicationCore/Private/Android/AndroidPlatformApplicationMisc.cpp
+++ b/Engine/Source/Runtime/ApplicationCore/Private/Android/AndroidPlatformApplicationMisc.cpp
@@ -227,7 +227,7 @@ EScreenPhysicalAccuracy FAndroidApplicationMisc::ComputePhysicalScreenDensity(in
 
 #if USE_ANDROID_JNI
 	extern FString AndroidThunkCpp_GetMetaDataString(const FString& Key);
-	FString DPIStrings = AndroidThunkCpp_GetMetaDataString(TEXT("ue4.displaymetrics.dpi"));
+	FString DPIStrings = AndroidThunkCpp_GetMetaDataString(TEXT("unreal.displaymetrics.dpi"));
 	TArray<FString> DPIValues;
 	DPIStrings.ParseIntoArray(DPIValues, TEXT(","));
 
diff --git a/Engine/Source/Runtime/ApplicationCore/Private/Android/AndroidWindow.cpp b/Engine/Source/Runtime/ApplicationCore/Private/Android/AndroidWindow.cpp
--- a/Engine/Source/Runtime/ApplicationCore/Private/Android/AndroidWindow.cpp
+++ b/Engine/Source/Runtime/ApplicationCore/Private/Android/AndroidWindow.cpp
@@ -509,7 +509,7 @@ void FAndroidWindow::CalculateSurfaceSize(int32_t& SurfaceWidth, int32_t& Surfac
 
 	// ensure the size is divisible by a specified amount
 	// do not convert to a surface size that is larger than native resolution
-	// Mobile VR doesn't need buffer quantization as UE4 never renders directly to the buffer in VR mode. 
+	// Mobile VR doesn't need buffer quantization as Unreal never renders directly to the buffer in VR mode. 
 	static const bool bIsMobileVRApp = AndroidThunkCpp_IsOculusMobileApplication() || FAndroidMisc::IsDaydreamApplication();
 	const int DividableBy = bIsMobileVRApp ? 1 : 8;
 	SurfaceWidth = (SurfaceWidth / DividableBy) * DividableBy;
diff --git a/Engine/Source/Runtime/ApplicationCore/Private/IOS/IOSAppDelegate.cpp b/Engine/Source/Runtime/ApplicationCore/Private/IOS/IOSAppDelegate.cpp
--- a/Engine/Source/Runtime/ApplicationCore/Private/IOS/IOSAppDelegate.cpp
+++ b/Engine/Source/Runtime/ApplicationCore/Private/IOS/IOSAppDelegate.cpp
@@ -237,7 +237,7 @@ static IOSAppDelegate* CachedDelegate = nil;
 #if BUILD_EMBEDDED_APP
 	if (CachedDelegate == nil)
 	{
-		UE_LOG(LogIOS, Fatal, TEXT("Currently, a native embedding UE4 must have the AppDelegate subclass from IOSAppDelegate."));
+		UE_LOG(LogIOS, Fatal, TEXT("Currently, a native app embedding Unreal must have the AppDelegate subclass from IOSAppDelegate."));
 
 		// if we are embedded, but CachedDelegate is nil, then that means the delegate was not an IOSAppDelegate subclass,
 		// so we need to do a switcheroo - but this is unlikely to work well
diff --git a/Engine/Source/Runtime/ApplicationCore/Private/IOS/IOSInputInterface.cpp b/Engine/Source/Runtime/ApplicationCore/Private/IOS/IOSInputInterface.cpp
--- a/Engine/Source/Runtime/ApplicationCore/Private/IOS/IOSInputInterface.cpp
+++ b/Engine/Source/Runtime/ApplicationCore/Private/IOS/IOSInputInterface.cpp
@@ -663,7 +663,7 @@ void FIOSInputInterface::GetMovementData(FVector& Attitude, FVector& RotationRat
 			[CurrentAttitude multiplyByInverseOfAttitude : ReferenceAttitude];
 		}
 
-		// convert to UE3
+		// convert to Unreal coordinate system
 		Attitude = FVector(float(CurrentAttitude.pitch), float(CurrentAttitude.yaw), float(CurrentAttitude.roll));
 		RotationRate = FVector(float(CurrentRotationRate.x), float(CurrentRotationRate.y), float(CurrentRotationRate.z));
 		Gravity = FVector(float(CurrentGravity.x), float(CurrentGravity.y), float(CurrentGravity.z));
diff --git a/Engine/Source/Runtime/Core/Private/Android/AndroidPlatformMemory.cpp b/Engine/Source/Runtime/Core/Private/Android/AndroidPlatformMemory.cpp
--- a/Engine/Source/Runtime/Core/Private/Android/AndroidPlatformMemory.cpp
+++ b/Engine/Source/Runtime/Core/Private/Android/AndroidPlatformMemory.cpp
@@ -200,7 +200,7 @@ FPlatformMemoryStats FAndroidPlatformMemory::GetStats()
 	// note: Android 10 places impractical limits on the frequency of calls to getProcessMemoryInfo, revert to VmRSS for OS10+
 	if (GJavaVM && FAndroidMisc::GetAndroidBuildVersion() < 29) 
 	{
-		MemoryStats.UsedPhysical = static_cast<uint64>(AndroidThunkCpp_GetMetaDataInt(TEXT("ue4.getUsedMemory"))) * 1024ULL;
+		MemoryStats.UsedPhysical = static_cast<uint64>(AndroidThunkCpp_GetMetaDataInt(TEXT("unreal.getUsedMemory"))) * 1024ULL;
 	}
 #endif
 
@@ -214,7 +214,7 @@ uint64 FAndroidPlatformMemory::GetMemoryUsedFast()
 	// note: Android 10 places impractical limits on the frequency of calls to getProcessMemoryInfo, revert to VmRSS for OS10+
 	if (GJavaVM && FAndroidMisc::GetAndroidBuildVersion() < 29) 
 	{
-		return static_cast<uint64>(AndroidThunkCpp_GetMetaDataInt(TEXT("ue4.getUsedMemory"))) * 1024ULL;
+		return static_cast<uint64>(AndroidThunkCpp_GetMetaDataInt(TEXT("unreal.getUsedMemory"))) * 1024ULL;
 	}
 #endif
 
diff --git a/Engine/Source/Runtime/Launch/Private/Android/AndroidJNI.cpp b/Engine/Source/Runtime/Launch/Private/Android/AndroidJNI.cpp
--- a/Engine/Source/Runtime/Launch/Private/Android/AndroidJNI.cpp
+++ b/Engine/Source/Runtime/Launch/Private/Android/AndroidJNI.cpp
@@ -1492,7 +1492,7 @@ void AndroidThunkCpp_SetDesiredViewSize(int32 Width, int32 Height)
 
 	// also send info to native wrapper around embedded
 	FEmbeddedCallParamsHelper Helper;
-	Helper.Command = TEXT("setue4resolution");
+	Helper.Command = TEXT("setueresolution");
 	Helper.Parameters = { {TEXT("width"), LexToString(Width)}, {TEXT("height"), LexToString(Height)} };
 	FEmbeddedDelegates::GetEmbeddedToNativeParamsDelegateForSubsystem(TEXT("native")).Broadcast(Helper);
 
diff --git a/Engine/Source/Runtime/Launch/Private/IOS/LaunchIOS.cpp b/Engine/Source/Runtime/Launch/Private/IOS/LaunchIOS.cpp
--- a/Engine/Source/Runtime/Launch/Private/IOS/LaunchIOS.cpp
+++ b/Engine/Source/Runtime/Launch/Private/IOS/LaunchIOS.cpp
@@ -336,13 +336,13 @@ static void MainThreadInit()
 
 	CGRect FullResolutionRect = MainFrame;
 
-	// embedded apps are embedded inside a UE4 view, so it's already made
+	// embedded apps are embedded inside a UE view, so it's already made
 #if BUILD_EMBEDDED_APP
 	// tell the embedded app that the .ini files are ready to be used, ie the View can be made if it was waiting to create the view
 	FEmbeddedCallParamsHelper Helper;
 	Helper.Command = TEXT("inisareready");
 	FEmbeddedDelegates::GetEmbeddedToNativeParamsDelegateForSubsystem(TEXT("native")).Broadcast(Helper);
-	// checkf(AppDelegate.IOSView != nil, TEXT("For embedded apps, the UE4EmbeddedView must have been created and set into the AppDelegate as IOSView"));
+	// checkf(AppDelegate.IOSView != nil, TEXT("For embedded apps, the UEEmbeddedView must have been created and set into the AppDelegate as IOSView"));
 #else
 	AppDelegate.IOSView = [[FIOSView alloc] initWithFrame:FullResolutionRect];
 	AppDelegate.IOSView.clearsContextBeforeDrawing = NO;
diff --git a/Engine/Source/Runtime/Launch/Private/IOS/TcpConsoleListener.cpp b/Engine/Source/Runtime/Launch/Private/IOS/TcpConsoleListener.cpp
--- a/Engine/Source/Runtime/Launch/Private/IOS/TcpConsoleListener.cpp
+++ b/Engine/Source/Runtime/Launch/Private/IOS/TcpConsoleListener.cpp
@@ -16,7 +16,7 @@ TcpConsoleListener::TcpConsoleListener(const FIPv4Endpoint& InListenEndpoint)
 	, bStopping(false)
 	, Listener(nullptr)
 {
-	UE_LOG(LogTemp, Warning, TEXT("[UE4] Console Listener created!\n"));
+	UE_LOG(LogTemp, Warning, TEXT("[UE] Console Listener created!\n"));
 	Thread = FRunnableThread::Create(this, TEXT("TcpConsoleListener"), 128 * 1024, TPri_Normal);
 }
 
@@ -54,7 +54,7 @@ void TcpConsoleListener::Exit()
 
 bool TcpConsoleListener::Init()
 {
-	UE_LOG(LogTemp, Warning, TEXT("[UE4] TCP Listener created!\n"));
+	UE_LOG(LogTemp, Warning, TEXT("[UE] TCP Listener created!\n"));
 	Listener = new FTcpListener(ListenEndpoint);
 	Listener->OnConnectionAccepted().BindRaw(this, &TcpConsoleListener::HandleListenerConnectionAccepted);
 
@@ -82,7 +82,7 @@ uint32 TcpConsoleListener::Run()
 				memset(RecvBuffer, 0, CommandSize);
 				if (Connection->Recv(RecvBuffer, CommandSize, BytesRead))
 				{
-					UE_LOG(LogTemp, Warning, TEXT("[UE4] Got TCP console command size %i '%s'"), BytesRead, *FString(UTF8_TO_TCHAR(RecvBuffer)));
+					UE_LOG(LogTemp, Warning, TEXT("[UE] Got TCP console command size %i '%s'"), BytesRead, *FString(UTF8_TO_TCHAR(RecvBuffer)));
 					EnqueueConsoleCommand(RecvBuffer);
 				}
 				else
diff --git a/Engine/Source/Runtime/Launch/Private/LaunchEngineLoop.cpp b/Engine/Source/Runtime/Launch/Private/LaunchEngineLoop.cpp
--- a/Engine/Source/Runtime/Launch/Private/LaunchEngineLoop.cpp
+++ b/Engine/Source/Runtime/Launch/Private/LaunchEngineLoop.cpp
@@ -2156,7 +2156,7 @@ int32 FEngineLoop::PreInitPreStartupScreen(const TCHAR* CmdLine)
 		if (!FPaths::IsProjectFilePathSet())
 		{
 			//@todo this is too early to localize
-			FMessageDialog::Open(EAppMsgType::Ok, NSLOCTEXT("Engine", "UE4RequiresProjectFiles", "UE4 games require a project file as the first parameter."));
+			FMessageDialog::Open(EAppMsgType::Ok, NSLOCTEXT("Engine", "UERequiresProjectFiles", "Unreal Engine games require a project file as the first parameter."));
 			return 1;
 		}
 	}
diff --git a/Engine/Source/Runtime/Launch/Private/Mac/LaunchMac.cpp b/Engine/Source/Runtime/Launch/Private/Mac/LaunchMac.cpp
--- a/Engine/Source/Runtime/Launch/Private/Mac/LaunchMac.cpp
+++ b/Engine/Source/Runtime/Launch/Private/Mac/LaunchMac.cpp
@@ -73,7 +73,7 @@ static int32 MacOSVersionCompare(const NSOperatingSystemVersion& VersionA, const
 	return 0;
 }
 
-@interface UE4AppDelegate : NSObject <NSApplicationDelegate, NSFileManagerDelegate>
+@interface UEAppDelegate : NSObject <NSApplicationDelegate, NSFileManagerDelegate>
 {
 #if WITH_EDITOR
 	NSString* Filename;
@@ -87,7 +87,7 @@ static int32 MacOSVersionCompare(const NSOperatingSystemVersion& VersionA, const
 
 @end
 
-@implementation UE4AppDelegate
+@implementation UEAppDelegate
 
 - (void)awakeFromNib
 {
@@ -408,7 +408,7 @@ INT32_MAIN_INT32_ARGC_TCHAR_ARGV()
 
 	SCOPED_AUTORELEASE_POOL;
 	[NSApplication sharedApplication];
-	[NSApp setDelegate:[UE4AppDelegate new]];
+	[NSApp setDelegate:[UEAppDelegate new]];
 	[NSApp run];
 	return GGuardedMainErrorLevel;
 }
diff --git a/Engine/Source/Runtime/Launch/Resources/Version.h b/Engine/Source/Runtime/Launch/Resources/Version.h
--- a/Engine/Source/Runtime/Launch/Resources/Version.h
+++ b/Engine/Source/Runtime/Launch/Resources/Version.h
@@ -3,14 +3,14 @@
 #pragma once
 
 //=============================================================================================================================================================
-// This file defines constants used for versioning packages, modules, and various subsystems in UE4. In general, it should not be necessary to include this
+// This file defines constants used for versioning packages, modules, and various subsystems in UE. In general, it should not be necessary to include this
 // file and access these values directly - they are all wrapped behind the higher-level abstractions in FEngineVersion and the FApp class.
 //
-// The following concepts are used for versioning in UE4:
+// The following concepts are used for versioning in UE:
 //
 //  -  The *engine version* defines the explicit major/minor/patch version of the engine, plus the changelist and branch name that it was built from. The 
 //     changelist is assumed to be a monotonically increasing number in the current branch, and is used both as a unique identifier and to infer that one engine
-//     was later than another. Tagged property serialization in UE4 is tolerant to properties being added or removed, so we always want to prevent an older 
+//     was later than another. Tagged property serialization in UE is tolerant to properties being added or removed, so we always want to prevent an older 
 //     build of the engine loading assets created with a newer build, discarding properties which have recently been added, and silently losing data
 //     when the asset is saved out. The changelist allows ordering versions in such cases. The engine version is encapsulated by the FEngineVersion class, of 
 //     which there are two commonly referenced instances:
@@ -45,15 +45,15 @@
 //
 //  -  The *network version* and *replay version* are used for versioning the network and replay subsystems, and default to the compatible engine version.
 //
-//  -  The *engine association* in a .uproject file often takes the appearance of a version number for launcher-installed binary UE4 releases, but may be 
+//  -  The *engine association* in a .uproject file often takes the appearance of a version number for launcher-installed binary UE releases, but may be 
 //     other identifiers as well. See ProjectDescriptor.h for a description of how this technique works.
 //
 // Constants in this file are updated by AutomationTool and UnrealGameSync. Be careful when changing formatting for the submitted version of this file that 
 // these tools can still parse it.
 //=============================================================================================================================================================
 
-// These numbers define the banner UE4 version, and are the most significant numbers when ordering two engine versions (that is, a 4.12.* version is always 
-// newer than a 4.11.* version, regardless of the changelist that it was built with)
+// These numbers define the banner UE version, and are the most significant numbers when ordering two engine versions (that is, a 5.12.* version is always 
+// newer than a 5.11.* version, regardless of the changelist that it was built with)
 #define ENGINE_MAJOR_VERSION	5
 #define ENGINE_MINOR_VERSION	0
 #define ENGINE_PATCH_VERSION	0
diff --git a/Engine/Source/Runtime/Online/HTTP/Private/Android/AndroidPlatformHttp.cpp b/Engine/Source/Runtime/Online/HTTP/Private/Android/AndroidPlatformHttp.cpp
--- a/Engine/Source/Runtime/Online/HTTP/Private/Android/AndroidPlatformHttp.cpp
+++ b/Engine/Source/Runtime/Online/HTTP/Private/Android/AndroidPlatformHttp.cpp
@@ -32,8 +32,8 @@ TOptional<FString> FAndroidPlatformHttp::GetOperatingSystemProxyAddress()
 	extern int32 AndroidThunkCpp_GetMetaDataInt(const FString& Key);
 	extern FString AndroidThunkCpp_GetMetaDataString(const FString& Key);
 
-	FString ProxyHost = AndroidThunkCpp_GetMetaDataString(TEXT("ue4.http.proxy.proxyHost"));
-	int32 ProxyPort = AndroidThunkCpp_GetMetaDataInt(TEXT("ue4.http.proxy.proxyPort"));
+	FString ProxyHost = AndroidThunkCpp_GetMetaDataString(TEXT("unreal.http.proxy.proxyHost"));
+	int32 ProxyPort = AndroidThunkCpp_GetMetaDataInt(TEXT("unreal.http.proxy.proxyPort"));
 
 	if (ProxyPort != -1 && !ProxyHost.IsEmpty())
 	{
diff --git a/Engine/Source/Runtime/OpenGLDrv/Private/Android/AndroidOpenGLFramePacer.cpp b/Engine/Source/Runtime/OpenGLDrv/Private/Android/AndroidOpenGLFramePacer.cpp
--- a/Engine/Source/Runtime/OpenGLDrv/Private/Android/AndroidOpenGLFramePacer.cpp
+++ b/Engine/Source/Runtime/OpenGLDrv/Private/Android/AndroidOpenGLFramePacer.cpp
@@ -268,7 +268,7 @@ bool FAndroidOpenGLFramePacer::SwapBuffers(bool bLockToVsync)
 				);
 			}
 
-			float RefreshRate = AndroidThunkCpp_GetMetaDataFloat(TEXT("ue4.display.getRefreshRate"));
+			float RefreshRate = AndroidThunkCpp_GetMetaDataFloat(TEXT("unreal.display.getRefreshRate"));
 
 			UE_LOG(LogRHI, Log, TEXT("JNI Display getRefreshRate=%f"),
 				RefreshRate
