Opt in to the new D3D12 version when it's available

diff --git a/Engine/Source/Runtime/Launch/Private/Windows/LaunchWindows.cpp b/Engine/Source/Runtime/Launch/Private/Windows/LaunchWindows.cpp
--- a/Engine/Source/Runtime/Launch/Private/Windows/LaunchWindows.cpp
+++ b/Engine/Source/Runtime/Launch/Private/Windows/LaunchWindows.cpp
@@ -32,8 +32,13 @@ extern "C" { _declspec(dllexport) uint32 AmdPowerXpressRequestHighPerformance =
 
 // Opt in to new D3D12 redist and tell the loader where to search for D3D12Core.dll.
 // The D3D loader looks for these symbol exports in the .exe module.
-// extern "C" { _declspec(dllexport) extern const UINT D3D12SDKVersion = 4; }
-// extern "C" { _declspec(dllexport) extern const char* D3D12SDKPath = u8".\\D3D12\\"; }
+// We only support this on x64 Windows Desktop platforms. Other platforms or non-redist-aware 
+// versions of Windows will transparently load default OS-provided D3D12 library.
+#define USE_D3D12_REDIST (PLATFORM_DESKTOP && PLATFORM_CPU_X86_FAMILY && PLATFORM_64BITS && 1)
+#if USE_D3D12_REDIST
+extern "C" { _declspec(dllexport) extern const UINT D3D12SDKVersion = 4; }
+extern "C" { _declspec(dllexport) extern const char* D3D12SDKPath = u8".\\D3D12\\"; }
+#endif // USE_D3D12_REDIST
 
 /**
  * Maintain a named mutex to detect whether we are the first instance of this game
