Disabling d3d12 redist until we find out what's wrong with the latest Nvidia driver.

diff --git a/Engine/Source/Runtime/Launch/Private/Windows/LaunchWindows.cpp b/Engine/Source/Runtime/Launch/Private/Windows/LaunchWindows.cpp
--- a/Engine/Source/Runtime/Launch/Private/Windows/LaunchWindows.cpp
+++ b/Engine/Source/Runtime/Launch/Private/Windows/LaunchWindows.cpp
@@ -34,7 +34,7 @@ extern "C" { _declspec(dllexport) uint32 AmdPowerXpressRequestHighPerformance =
 // The D3D loader looks for these symbol exports in the .exe module.
 // We only support this on x64 Windows Desktop platforms. Other platforms or non-redist-aware 
 // versions of Windows will transparently load default OS-provided D3D12 library.
-#define USE_D3D12_REDIST (PLATFORM_DESKTOP && PLATFORM_CPU_X86_FAMILY && PLATFORM_64BITS && 1)
+#define USE_D3D12_REDIST (PLATFORM_DESKTOP && PLATFORM_CPU_X86_FAMILY && PLATFORM_64BITS && 0)
 #if USE_D3D12_REDIST
 extern "C" { _declspec(dllexport) extern const UINT D3D12SDKVersion = 4; }
 extern "C" { _declspec(dllexport) extern const char* D3D12SDKPath = u8".\\D3D12\\"; }
