Added Webauth plugin to QAGame
Fixed bug with LaunchDaemonMessages library not loading

diff --git a/Engine/Source/Runtime/Launch/Private/IOS/GameLaunchDaemonMessageHandler.cpp b/Engine/Source/Runtime/Launch/Private/IOS/GameLaunchDaemonMessageHandler.cpp
--- a/Engine/Source/Runtime/Launch/Private/IOS/GameLaunchDaemonMessageHandler.cpp
+++ b/Engine/Source/Runtime/Launch/Private/IOS/GameLaunchDaemonMessageHandler.cpp
@@ -14,6 +14,8 @@
 
 void FGameLaunchDaemonMessageHandler::Init()
 {
+    FModuleManager::Get().LoadModuleChecked(TEXT("LaunchDaemonMessages"));
+    
 	MessageEndpoint = FMessageEndpoint::Builder("FGameLaunchDaemonMessageHandler")
 		.Handling<FIOSLaunchDaemonPing>(this, &FGameLaunchDaemonMessageHandler::HandlePingMessage)
 		.Handling<FIOSLaunchDaemonLaunchApp>(this, &FGameLaunchDaemonMessageHandler::HandleLaunchRequest);
