added the call to FThreadStats::StopThread() to FEngineLoop::AppExit(). the problem is that StopThread() is called in FEngineLoop::Exit() which is not used in some programs, e.g. BlankProgram, and potentially many programs used it as a template.

StopThread() waits until all stat tasks are executed. Now that we use a pipe instead of StatsThread, the pipe asserts in its destructor if there're any unexecuted tasks. Seems like before unexecuted tasks were just dropped.

diff --git a/Engine/Source/Runtime/Launch/Private/LaunchEngineLoop.cpp b/Engine/Source/Runtime/Launch/Private/LaunchEngineLoop.cpp
--- a/Engine/Source/Runtime/Launch/Private/LaunchEngineLoop.cpp
+++ b/Engine/Source/Runtime/Launch/Private/LaunchEngineLoop.cpp
@@ -5997,8 +5997,11 @@ void FEngineLoop::AppPreExit( )
 
 void FEngineLoop::AppExit()
 {
-#if !WITH_ENGINE
 	// when compiled WITH_ENGINE, this will happen in FEngineLoop::Exit()
+#if !WITH_ENGINE
+#if STATS
+	FThreadStats::StopThread();
+#endif
 	FTaskGraphInterface::Shutdown();
 #endif // WITH_ENGINE
 
