Add the SubobjectDataInterface to the LaunchEngineLoop for use in Blueprints, Python, and unit tests

#rnx

diff --git a/Engine/Plugins/Tests/EditorTests/Source/EditorTests/EditorTests.Build.cs b/Engine/Plugins/Tests/EditorTests/Source/EditorTests/EditorTests.Build.cs
--- a/Engine/Plugins/Tests/EditorTests/Source/EditorTests/EditorTests.Build.cs
+++ b/Engine/Plugins/Tests/EditorTests/Source/EditorTests/EditorTests.Build.cs
@@ -69,6 +69,7 @@ public class EditorTests : ModuleRules
                 "MeshBuilder",
                 "RawMesh",
 				"AutomationController",
+				//"SubobjectDataInterface",
 			}
 		);
 
diff --git a/Engine/Source/Editor/Kismet/Kismet.Build.cs b/Engine/Source/Editor/Kismet/Kismet.Build.cs
--- a/Engine/Source/Editor/Kismet/Kismet.Build.cs
+++ b/Engine/Source/Editor/Kismet/Kismet.Build.cs
@@ -58,7 +58,8 @@ public class Kismet : ModuleRules
                 "UMG", // for SBlueprintDiff
                 "WorkspaceMenuStructure",
 				"ToolMenus",
-            }
+				"SubobjectDataInterface",
+			}
 			);
 
         DynamicallyLoadedModuleNames.AddRange(
diff --git a/Engine/Source/Runtime/Launch/Private/LaunchEngineLoop.cpp b/Engine/Source/Runtime/Launch/Private/LaunchEngineLoop.cpp
--- a/Engine/Source/Runtime/Launch/Private/LaunchEngineLoop.cpp
+++ b/Engine/Source/Runtime/Launch/Private/LaunchEngineLoop.cpp
@@ -3929,6 +3929,7 @@ bool FEngineLoop::LoadStartupCoreModules()
 	FModuleManager::Get().LoadModuleChecked("UnrealEd");
 	FModuleManager::LoadModuleChecked<IEditorStyleModule>("EditorStyle");
 	FModuleManager::Get().LoadModuleChecked("LandscapeEditorUtilities");
+	//FModuleManager::Get().LoadModuleChecked("SubobjectDataInterface");
 #endif //WITH_EDITOR
 
 	// Load UI modules
