Pregenerate autogen headers for target platforms (UE-115409).

- Will hopefully prevent races between the generation and SCWs.
- Startup impact: ~2ms (when files are up to date) to ~7ms (when they aren't).

diff --git a/Engine/Source/Runtime/Engine/Private/ShaderCompiler/ShaderGenerationUtil.cpp b/Engine/Source/Runtime/Engine/Private/ShaderCompiler/ShaderGenerationUtil.cpp
--- a/Engine/Source/Runtime/Engine/Private/ShaderCompiler/ShaderGenerationUtil.cpp
+++ b/Engine/Source/Runtime/Engine/Private/ShaderCompiler/ShaderGenerationUtil.cpp
@@ -2038,6 +2038,12 @@ void FShaderCompileUtilities::GenerateBrdfHeaders(EShaderPlatform TargetPlatform
 	FGBufferInfo Info = FetchGBufferInfoAndWriteAutogen(TargetPlatform, FeatureLevel);
 }
 
+void FShaderCompileUtilities::GenerateBrdfHeaders(const FName& ShaderFormat)
+{
+	const EShaderPlatform ShaderPlatform = ShaderFormatToLegacyShaderPlatform(ShaderFormat);
+	FShaderCompileUtilities::GenerateBrdfHeaders(ShaderPlatform);
+}
+
 ENGINE_API FGBufferParams FShaderCompileUtilities::FetchGBufferParamsRuntime(EShaderPlatform Platform)
 {
 
diff --git a/Engine/Source/Runtime/Engine/Public/ShaderCompiler.h b/Engine/Source/Runtime/Engine/Public/ShaderCompiler.h
--- a/Engine/Source/Runtime/Engine/Public/ShaderCompiler.h
+++ b/Engine/Source/Runtime/Engine/Public/ShaderCompiler.h
@@ -633,7 +633,8 @@ namespace FShaderCompileUtilities
 	void MoveFileHelper(const FString& To, const FString& From);
 	void DeleteFileHelper(const FString& Filename);
 
-	void GenerateBrdfHeaders(const EShaderPlatform Platform);
+	ENGINE_API void GenerateBrdfHeaders(const EShaderPlatform Platform);
+	ENGINE_API void GenerateBrdfHeaders(const FName& ShaderFormat);
 	void ApplyDerivedDefines(FShaderCompilerEnvironment& OutEnvironment, FShaderCompilerEnvironment * SharedEnvironment, const EShaderPlatform Platform);
 	void AppendGBufferDDCKeyString(const EShaderPlatform Platform, FString& KeyString);
 	ENGINE_API FGBufferInfo FetchGBufferInfoAndWriteAutogen(EShaderPlatform TargetPlatform, ERHIFeatureLevel::Type FeatureLevel);
diff --git a/Engine/Source/Runtime/Launch/Private/LaunchEngineLoop.cpp b/Engine/Source/Runtime/Launch/Private/LaunchEngineLoop.cpp
--- a/Engine/Source/Runtime/Launch/Private/LaunchEngineLoop.cpp
+++ b/Engine/Source/Runtime/Launch/Private/LaunchEngineLoop.cpp
@@ -1468,6 +1468,33 @@ static void UpdateGInputTime()
 	GInputTime = FPlatformTime::Cycles64();
 }
 
+#if WITH_EDITOR
+static void ShaderAutogenInit()
+{
+	// Force generation of the autogen files for the host platform
+	FShaderCompileUtilities::GenerateBrdfHeaders(GMaxRHIShaderPlatform);
+
+	// also do this for all active target platforms (e.g. when cooking)
+	ITargetPlatformManagerModule* TPM = GetTargetPlatformManager();
+	if (TPM)
+	{
+		const TArray<ITargetPlatform*>& Platforms = TPM->GetActiveTargetPlatforms();
+
+		for (int32 Index = 0; Index < Platforms.Num(); ++Index)
+		{
+			TArray<FName> DesiredShaderFormats;
+			checkf(Platforms[Index], TEXT("Null platform on the list of active platforms!"));
+			Platforms[Index]->GetAllTargetedShaderFormats(DesiredShaderFormats);
+
+			for (int32 FormatIndex = 0; FormatIndex < DesiredShaderFormats.Num(); ++FormatIndex)
+			{
+				FShaderCompileUtilities::GenerateBrdfHeaders(DesiredShaderFormats[FormatIndex]);
+			}
+		}
+	}
+}
+#endif // WITH_EDITOR
+
 DECLARE_CYCLE_STAT(TEXT("FEngineLoop::PreInitPreStartupScreen.AfterStats"), STAT_FEngineLoop_PreInitPreStartupScreen_AfterStats, STATGROUP_LoadTime);
 DECLARE_CYCLE_STAT(TEXT("FEngineLoop::PreInitPostStartupScreen.AfterStats"), STAT_FEngineLoop_PreInitPostStartupScreen_AfterStats, STATGROUP_LoadTime);
 
@@ -2741,8 +2768,8 @@ PRAGMA_ENABLE_DEPRECATION_WARNINGS
 
 #if WITH_EDITOR
 	{
-		// Force generation of the autogen files if they don't already exist. We'll ignore the returned data.
-		FGBufferInfo GBufferInfo = FShaderCompileUtilities::FetchGBufferInfoAndWriteAutogen(GMaxRHIShaderPlatform, GetMaxSupportedFeatureLevel(GMaxRHIShaderPlatform));
+		// pre-generate certain shader code based on the engine configuration (editor only)
+		ShaderAutogenInit();
 	}
 #endif
 
