DDC: Ensure that the build system is constructed on launch alongside the cache

#rnx

diff --git a/Engine/Source/Runtime/Launch/Private/LaunchEngineLoop.cpp b/Engine/Source/Runtime/Launch/Private/LaunchEngineLoop.cpp
--- a/Engine/Source/Runtime/Launch/Private/LaunchEngineLoop.cpp
+++ b/Engine/Source/Runtime/Launch/Private/LaunchEngineLoop.cpp
@@ -2806,6 +2806,7 @@ PRAGMA_ENABLE_DEPRECATION_WARNINGS
 	{
 		// Ensure that DDC is initialized from the game thread.
 		GetDerivedDataCacheRef();
+		GetDerivedDataBuildRef();
 
 		check(!GDistanceFieldAsyncQueue);
 		GDistanceFieldAsyncQueue = new FDistanceFieldAsyncQueue();
