Enable the loading of the SubobjectDataInterface in the LaunchEngineLoop

#rnx

diff --git a/Engine/Source/Runtime/Launch/Private/LaunchEngineLoop.cpp b/Engine/Source/Runtime/Launch/Private/LaunchEngineLoop.cpp
--- a/Engine/Source/Runtime/Launch/Private/LaunchEngineLoop.cpp
+++ b/Engine/Source/Runtime/Launch/Private/LaunchEngineLoop.cpp
@@ -3957,7 +3957,7 @@ bool FEngineLoop::LoadStartupCoreModules()
 	FModuleManager::Get().LoadModuleChecked("UnrealEd");
 	FModuleManager::LoadModuleChecked<IEditorStyleModule>("EditorStyle");
 	FModuleManager::Get().LoadModuleChecked("LandscapeEditorUtilities");
-	//FModuleManager::Get().LoadModuleChecked("SubobjectDataInterface");
+	FModuleManager::Get().LoadModuleChecked("SubobjectDataInterface");
 #endif //WITH_EDITOR
 
 	// Load UI modules
