Reinstate auto-connect for the editor proper. This was previously disabled if the editor was being run as it would interfer with plugins that would create their own store and enable tracing to it.

#rnx

diff --git a/Engine/Source/Runtime/Launch/Private/LaunchEngineLoop.cpp b/Engine/Source/Runtime/Launch/Private/LaunchEngineLoop.cpp
--- a/Engine/Source/Runtime/Launch/Private/LaunchEngineLoop.cpp
+++ b/Engine/Source/Runtime/Launch/Private/LaunchEngineLoop.cpp
@@ -1586,6 +1586,7 @@ int32 FEngineLoop::PreInitPreStartupScreen(const TCHAR* CmdLine)
 
 	// Initialize trace
 	FTraceAuxiliary::Initialize(CmdLine);
+	FTraceAuxiliary::TryAutoConnect();
 
 	// disable/enable LLM based on commandline
 	{
@@ -2091,13 +2092,6 @@ int32 FEngineLoop::PreInitPreStartupScreen(const TCHAR* CmdLine)
 #endif
 #endif	//UE_EDITOR
 
-#if !UE_BUILD_SHIPPING && !IS_PROGRAM
-	if (!bHasEditorToken)
-	{
-		FTraceAuxiliary::TryAutoConnect();
-	}
-#endif
-
 #if !UE_BUILD_SHIPPING
 	// Benchmarking.
 	FApp::SetBenchmarking(FParse::Param(FCommandLine::Get(), TEXT("BENCHMARK")));
