Make sure that the FVirtualizationManager is created on the game thread as some of it's dependancies now require being created on the game thread.

- This will be cleaned up further when FVirtualizationManager is moved into the Virtualization module and we do a pass on the initialization.

#rnx

diff --git a/Engine/Source/Runtime/Launch/Private/LaunchEngineLoop.cpp b/Engine/Source/Runtime/Launch/Private/LaunchEngineLoop.cpp
--- a/Engine/Source/Runtime/Launch/Private/LaunchEngineLoop.cpp
+++ b/Engine/Source/Runtime/Launch/Private/LaunchEngineLoop.cpp
@@ -107,6 +107,7 @@
 	#include "AssetCompilingManager.h"
 	#include "Serialization/BulkDataRegistry.h"
 	#include "ShaderCompiler.h"
+	#include "Virtualization/VirtualizationManager.h"
 
 	#if PLATFORM_WINDOWS
 		#include "Windows/AllowWindowsPlatformTypes.h"
@@ -2841,6 +2842,11 @@ PRAGMA_ENABLE_DEPRECATION_WARNINGS
 		GetDerivedDataCacheRef();
 		GetDerivedDataBuildRef();
 
+#if WITH_EDITOR
+		// Make sure that FVirtualizationManager is first called from the game thread
+		UE::Virtualization::FVirtualizationManager::Get();
+#endif //WITH_EDITOR
+
 		check(!GDistanceFieldAsyncQueue);
 		GDistanceFieldAsyncQueue = new FDistanceFieldAsyncQueue();
 
