Fixed case where editor starting in a mobile preview will crash failing to find auto-generated shaders

diff --git a/Engine/Source/Runtime/Launch/Private/LaunchEngineLoop.cpp b/Engine/Source/Runtime/Launch/Private/LaunchEngineLoop.cpp
--- a/Engine/Source/Runtime/Launch/Private/LaunchEngineLoop.cpp
+++ b/Engine/Source/Runtime/Launch/Private/LaunchEngineLoop.cpp
@@ -1514,6 +1514,13 @@ static void ShaderAutogenInit()
 			}
 		}
 	}
+
+	// also do this for the editor mobile preview
+	EShaderPlatform MobilePreviewShaderPlatform = GShaderPlatformForFeatureLevel[ERHIFeatureLevel::ES3_1];
+	if (MobilePreviewShaderPlatform != SP_NumPlatforms)
+	{
+		FShaderCompileUtilities::GenerateBrdfHeaders(MobilePreviewShaderPlatform);
+	}
 }
 #endif // WITH_EDITOR
 
