Bug fix for crash on exit, IODispatcher needs to be shutdown after the shadercodelibrary

diff --git a/Engine/Source/Runtime/Launch/Private/LaunchEngineLoop.cpp b/Engine/Source/Runtime/Launch/Private/LaunchEngineLoop.cpp
--- a/Engine/Source/Runtime/Launch/Private/LaunchEngineLoop.cpp
+++ b/Engine/Source/Runtime/Launch/Private/LaunchEngineLoop.cpp
@@ -4540,6 +4540,12 @@ void FEngineLoop::Exit()
 	// Close shader code map, if any
 	FShaderCodeLibrary::Shutdown();
 
+#if USE_IO_DISPATCHER
+	// has to happen after the FShaderCodeLibrary, as it holds on to file requests
+	FIoDispatcher::Shutdown();
+#endif
+
+
 #if !PLATFORM_ANDROID || PLATFORM_LUMIN // UnloadModules doesn't work on Android
 #if WITH_ENGINE
 	// Save the hot reload state
@@ -6118,9 +6124,6 @@ void FEngineLoop::AppPreExit( )
 
 #endif
 
-#if USE_IO_DISPATCHER
-	FIoDispatcher::Shutdown();
-#endif
 }
 
 
