Fix typo which broke the iOS build.

#rnx

diff --git a/Engine/Source/Runtime/Launch/Private/IOS/GameLaunchDaemonMessageHandler.cpp b/Engine/Source/Runtime/Launch/Private/IOS/GameLaunchDaemonMessageHandler.cpp
--- a/Engine/Source/Runtime/Launch/Private/IOS/GameLaunchDaemonMessageHandler.cpp
+++ b/Engine/Source/Runtime/Launch/Private/IOS/GameLaunchDaemonMessageHandler.cpp
@@ -42,7 +42,7 @@ void FGameLaunchDaemonMessageHandler::HandlePingMessage(const FIOSLaunchDaemonPi
 	{
 		FMessageAddress MessageSender = Context->GetSender();
 
-		MessageEndpoint->Send(FMessageEndpoint::MakeMessage<<FIOSLaunchDaemonPong>(FString(FPlatformProperties::PlatformName()) + (TARGET_IPHONE_SIMULATOR ? FString(TEXT("Simulator:")) : FString(TEXT("@"))) + FString(FPlatformProcess::ComputerName())
+		MessageEndpoint->Send(FMessageEndpoint::MakeMessage<FIOSLaunchDaemonPong>(FString(FPlatformProperties::PlatformName()) + (TARGET_IPHONE_SIMULATOR ? FString(TEXT("Simulator:")) : FString(TEXT("@"))) + FString(FPlatformProcess::ComputerName())
 			, FPlatformProcess::ComputerName()
 			, "Game_Running"
 			, [[UIDevice currentDevice] userInterfaceIdiom] == UIUserInterfaceIdiomPhone? "Phone" : "Tablet"
