Adding new CSV stat to record unaccounted render thread time outside of the scene render function.

diff --git a/Engine/Source/Runtime/Launch/Private/LaunchEngineLoop.cpp b/Engine/Source/Runtime/Launch/Private/LaunchEngineLoop.cpp
--- a/Engine/Source/Runtime/Launch/Private/LaunchEngineLoop.cpp
+++ b/Engine/Source/Runtime/Launch/Private/LaunchEngineLoop.cpp
@@ -4899,11 +4899,15 @@ static inline void BeginFrameRenderThread(FRHICommandListImmediate& RHICmdList,
 	{
 		GEngine->SetRenderSubmitLatencyMarkerStart(CurrentFrameCounter);
 	});
+
+	FCsvProfiler::BeginExclusiveStat("RenderThreadOther");
 }
 
 
 static inline void EndFrameRenderThread(FRHICommandListImmediate& RHICmdList, uint64 CurrentFrameCounter)
 {
+	FCsvProfiler::EndExclusiveStat("RenderThreadOther");
+
 	RHICmdList.EnqueueLambda([CurrentFrameCounter](FRHICommandListImmediate& InRHICmdList)
 	{
 		GEngine->SetRenderSubmitLatencyMarkerEnd(CurrentFrameCounter);
