Fixed CSV profiler build break.

diff --git a/Engine/Source/Runtime/Launch/Private/LaunchEngineLoop.cpp b/Engine/Source/Runtime/Launch/Private/LaunchEngineLoop.cpp
--- a/Engine/Source/Runtime/Launch/Private/LaunchEngineLoop.cpp
+++ b/Engine/Source/Runtime/Launch/Private/LaunchEngineLoop.cpp
@@ -4900,13 +4900,17 @@ static inline void BeginFrameRenderThread(FRHICommandListImmediate& RHICmdList,
 		GEngine->SetRenderSubmitLatencyMarkerStart(CurrentFrameCounter);
 	});
 
+#if CSV_PROFILER
 	FCsvProfiler::BeginExclusiveStat("RenderThreadOther");
+#endif
 }
 
 
 static inline void EndFrameRenderThread(FRHICommandListImmediate& RHICmdList, uint64 CurrentFrameCounter)
 {
+#if CSV_PROFILER
 	FCsvProfiler::EndExclusiveStat("RenderThreadOther");
+#endif
 
 	RHICmdList.EnqueueLambda([CurrentFrameCounter](FRHICommandListImmediate& InRHICmdList)
 	{
