Updating Engine Version to 5.1

#rnx

diff --git a/Engine/Build/Build.version b/Engine/Build/Build.version
--- a/Engine/Build/Build.version
+++ b/Engine/Build/Build.version
@@ -1,6 +1,6 @@
 {
 	"MajorVersion": 5,
-	"MinorVersion": 0,
+	"MinorVersion": 1,
 	"PatchVersion": 0,
 	"Changelist": 0,
 	"CompatibleChangelist": 16433597,
diff --git a/Engine/Source/Runtime/Launch/Resources/Version.h b/Engine/Source/Runtime/Launch/Resources/Version.h
--- a/Engine/Source/Runtime/Launch/Resources/Version.h
+++ b/Engine/Source/Runtime/Launch/Resources/Version.h
@@ -55,7 +55,7 @@
 // These numbers define the banner UE version, and are the most significant numbers when ordering two engine versions (that is, a 5.12.* version is always 
 // newer than a 5.11.* version, regardless of the changelist that it was built with)
 #define ENGINE_MAJOR_VERSION	5
-#define ENGINE_MINOR_VERSION	0
+#define ENGINE_MINOR_VERSION	1
 #define ENGINE_PATCH_VERSION	0
 
 // Macros for encoding strings
