#UE Fix crash that happens in non-golden-path scenario where you fail to startup completely due to a missing plugin (or maybe other reasons) and the streaming manager hasn't even been created yet. Normally its been fine to create it in the Getter here just so it can return immediately from BlockTillAllRequestsFinished, but there is now some code in the streamable manager assuming GConfig was set up, but it might not be at this time and it seems unncessary to create the manager here.

diff --git a/Engine/Source/Runtime/Launch/Private/LaunchEngineLoop.cpp b/Engine/Source/Runtime/Launch/Private/LaunchEngineLoop.cpp
--- a/Engine/Source/Runtime/Launch/Private/LaunchEngineLoop.cpp
+++ b/Engine/Source/Runtime/Launch/Private/LaunchEngineLoop.cpp
@@ -4406,7 +4406,10 @@ void FEngineLoop::Exit()
 	if (!IStreamingManager::HasShutdown())
 	{
 		UTexture2D::CancelPendingTextureStreaming();
-		IStreamingManager::Get().BlockTillAllRequestsFinished();
+		if (FStreamingManagerCollection* StreamingManager = IStreamingManager::Get_Concurrent())
+		{
+			StreamingManager->BlockTillAllRequestsFinished();
+		}
 	}
 	FAudioDeviceManager::Shutdown();
 
