MoviePlayer: DefaultGameMoviePlayer is now activated after FPreLoadScreenManager is done.
#rnx

diff --git a/Engine/Source/Runtime/Launch/Private/LaunchEngineLoop.cpp b/Engine/Source/Runtime/Launch/Private/LaunchEngineLoop.cpp
--- a/Engine/Source/Runtime/Launch/Private/LaunchEngineLoop.cpp
+++ b/Engine/Source/Runtime/Launch/Private/LaunchEngineLoop.cpp
@@ -5136,7 +5136,19 @@ void FEngineLoop::Tick()
                         GameEngine->SwitchGameWindowToUseGameViewport();
                     }
                 }
-                
+
+#if !UE_SERVER
+				// Is it ok to start up the movie player?
+				if (!IsRunningDedicatedServer() && !IsRunningCommandlet() && !GetMoviePlayer()->IsMovieCurrentlyPlaying())
+				{
+					// Enable the MoviePlayer now that the preload screen manager is done.
+					if (FSlateRenderer* Renderer = FSlateApplication::Get().GetRenderer())
+					{
+						GetMoviePlayer()->Initialize(*Renderer, FPreLoadScreenManager::Get()->GetRenderWindow());
+					}
+				}
+#endif // !UE_SERVER
+
                 //Destroy / Clean Up PreLoadScreenManager as we are now done
                 FPreLoadScreenManager::Destroy();
             }
