[CSVProfiler] - move initialization of ExtraDevMemory to an explicit call POST initialization of platform memory subsystems to avoid initialization order problems with the first call to the CSVProfiler singleton.

diff --git a/Engine/Source/Runtime/Core/Private/ProfilingDebugging/CsvProfiler.cpp b/Engine/Source/Runtime/Core/Private/ProfilingDebugging/CsvProfiler.cpp
--- a/Engine/Source/Runtime/Core/Private/ProfilingDebugging/CsvProfiler.cpp
+++ b/Engine/Source/Runtime/Core/Private/ProfilingDebugging/CsvProfiler.cpp
@@ -2654,11 +2654,6 @@ FCsvProfiler::FCsvProfiler()
 	{
 		SetMetadataInternal(TEXT("DeviceID"), *DeviceID);
 	}
-#if !UE_BUILD_SHIPPING
-	int32 ExtraDevelopmentMemoryMB = (int32)(FPlatformMemory::GetExtraDevelopmentMemorySize() / 1024ull / 1024ull);
-	SetMetadataInternal(TEXT("ExtraDevelopmentMemoryMB"), *FString::FromInt(ExtraDevelopmentMemoryMB));
-#endif
-
 }
 
 FCsvProfiler::~FCsvProfiler()
diff --git a/Engine/Source/Runtime/Launch/Private/LaunchEngineLoop.cpp b/Engine/Source/Runtime/Launch/Private/LaunchEngineLoop.cpp
--- a/Engine/Source/Runtime/Launch/Private/LaunchEngineLoop.cpp
+++ b/Engine/Source/Runtime/Launch/Private/LaunchEngineLoop.cpp
@@ -36,6 +36,7 @@
 #include "Stats/StatsMallocProfilerProxy.h"
 #include "Trace/Trace.inl"
 #include "ProfilingDebugging/TraceAuxiliary.h"
+#include "ProfilingDebugging/CsvProfiler.h"
 #include "ProfilingDebugging/BootProfiling.h"
 #if WITH_ENGINE
 #include "HAL/PlatformSplash.h"
@@ -2579,6 +2580,15 @@ PRAGMA_ENABLE_DEPRECATION_WARNINGS
 	}
 #endif // WITH_ENGINE
 
+#if !UE_BUILD_SHIPPING
+	{
+		int32 ExtraDevelopmentMemoryMB = (int32)(FPlatformMemory::GetExtraDevelopmentMemorySize() / 1024ull / 1024ull);
+		CSV_METADATA(TEXT("ExtraDevelopmentMemoryMB"), *FString::FromInt(ExtraDevelopmentMemoryMB));
+	}
+#endif
+
+#endif
+
 #if USE_IO_DISPATCHER
 	{
 		SCOPED_BOOT_TIMING("InitIoDispatcher");
