Remove extra endif

diff --git a/Engine/Source/Runtime/Launch/Private/LaunchEngineLoop.cpp b/Engine/Source/Runtime/Launch/Private/LaunchEngineLoop.cpp
--- a/Engine/Source/Runtime/Launch/Private/LaunchEngineLoop.cpp
+++ b/Engine/Source/Runtime/Launch/Private/LaunchEngineLoop.cpp
@@ -2587,8 +2587,6 @@ PRAGMA_ENABLE_DEPRECATION_WARNINGS
 	}
 #endif
 
-#endif
-
 #if USE_IO_DISPATCHER
 	{
 		SCOPED_BOOT_TIMING("InitIoDispatcher");
