Adding `MaxFrameTime` to the CSVProfiler output which will capture the `t.MaxFPS` value.

diff --git a/Engine/Source/Runtime/Launch/Private/LaunchEngineLoop.cpp b/Engine/Source/Runtime/Launch/Private/LaunchEngineLoop.cpp
--- a/Engine/Source/Runtime/Launch/Private/LaunchEngineLoop.cpp
+++ b/Engine/Source/Runtime/Launch/Private/LaunchEngineLoop.cpp
@@ -1445,6 +1445,16 @@ static void UpdateCoreCsvStats_EndFrame()
 	PhysicalMBFree -= float(FPlatformMemory::GetExtraDevelopmentMemorySize() / 1024ull / 1024ull);
 #endif
 	CSV_CUSTOM_STAT_GLOBAL(MemoryFreeMB, PhysicalMBFree, ECsvCustomStatOp::Set);
+
+#if !UE_BUILD_SHIPPING
+	float TargetFPS = 30.0f;
+	static IConsoleVariable* MaxFPSCVar = IConsoleManager::Get().FindConsoleVariable(TEXT("t.MaxFPS"));
+	if (MaxFPSCVar && MaxFPSCVar->GetFloat() > 0)
+	{
+		TargetFPS = MaxFPSCVar->GetFloat();
+	}
+	CSV_CUSTOM_STAT_GLOBAL(MaxFrameTime, 1000.0f / TargetFPS, ECsvCustomStatOp::Set);
+#endif
 }
 #endif // WITH_ENGINE && CSV_PROFILER
 
