Fix parsing of the -csvmetadata arg, so it doesn't stop on separators
Refactor csvExecCmds to move the initialization logic into the CSV profiler itself. The execution still needs to happen in the engine, but this is now done during normal frame processing.
Fix a bug when processing csv execcmds with colons, e.g file paths

diff --git a/Engine/Source/Runtime/Core/Private/ProfilingDebugging/CsvProfiler.cpp b/Engine/Source/Runtime/Core/Private/ProfilingDebugging/CsvProfiler.cpp
--- a/Engine/Source/Runtime/Core/Private/ProfilingDebugging/CsvProfiler.cpp
+++ b/Engine/Source/Runtime/Core/Private/ProfilingDebugging/CsvProfiler.cpp
@@ -179,6 +179,9 @@ bool GCsvProfilerNamedEventsExclusive = false;
 bool GCsvProfilerNamedEventsTiming = false;
 #endif //CSV_PROFILER_SUPPORT_NAMED_EVENTS
 
+
+static TMap<uint32, TArray<FString>>* GCsvFrameExecCmds = NULL;
+
 bool IsContinuousWriteEnabled(bool bGameThread)
 {
 	int CVarValue = -1;
@@ -509,6 +512,21 @@ int32 FCsvProfiler::RegisterCategory(const FString& CategoryName, bool bEnableBy
 	return FCsvCategoryData::Get()->RegisterCategory(CategoryName, bEnableByDefault, bIsGlobal);
 }
 
+void FCsvProfiler::GetFrameExecCommands(TArray<FString>& OutFrameCommands) const
+{
+	check(IsInGameThread());
+	OutFrameCommands.Empty();
+	if (GCsvProfilerIsCapturing && GCsvFrameExecCmds)
+	{
+		TArray<FString>* FrameCommands = GCsvFrameExecCmds->Find(CaptureFrameNumber);
+		if (FrameCommands)
+		{
+			OutFrameCommands = *FrameCommands;
+		}
+	}
+}
+
+
 
 bool IsInCsvProcessingThread()
 {
@@ -3461,7 +3479,7 @@ void FCsvProfiler::Init()
 	}
 
 	FString CsvMetadataStr;
-	if (FParse::Value(FCommandLine::Get(), TEXT("csvMetadata="), CsvMetadataStr))
+	if (FParse::Value(FCommandLine::Get(), TEXT("csvMetadata="), CsvMetadataStr, false))
 	{ 
 		TArray<FString> CsvMetadataList;
 		CsvMetadataStr.ParseIntoArray(CsvMetadataList, TEXT(","), true);
@@ -3522,6 +3540,35 @@ void FCsvProfiler::Init()
 		}
 	}
 	GCsvABTest.InitFromCommandline();
+
+	// Handle -csvExeccmds
+	FString CsvExecCommandsStr;
+	if (FParse::Value(FCommandLine::Get(), TEXT("-csvExecCmds="), CsvExecCommandsStr, false))
+	{
+		GCsvFrameExecCmds = new TMap<uint32, TArray<FString>>();
+
+		TArray<FString> CsvExecCommandsList;
+		if (CsvExecCommandsStr.ParseIntoArray(CsvExecCommandsList, TEXT(","), true) > 0)
+		{
+			for (FString FrameAndCommand : CsvExecCommandsList)
+			{
+				int32 ColonIndex = -1;
+				if (FrameAndCommand.FindChar(TEXT(':'), ColonIndex))
+				{
+					FString FrameStr = FrameAndCommand.Mid(0,ColonIndex);
+					FString CommandStr = FrameAndCommand.Mid(ColonIndex+1);
+					uint32 Frame = FCString::Atoi(*FrameStr);
+					if (!GCsvFrameExecCmds->Find(Frame))
+					{
+						GCsvFrameExecCmds->Add(Frame, TArray<FString>());
+					}
+					(*GCsvFrameExecCmds)[Frame].Add(CommandStr);
+					UE_LOG(LogCsvProfiler, Display, TEXT("Added CsvExecCommand - frame %d : %s"), Frame, *CommandStr);
+				}
+			}
+		}
+	}
+
 #endif // CSV_PROFILER_ALLOW_DEBUG_FEATURES
 
 	// Always disable the CSV profiling thread if the platform does not support threading.
@@ -3721,4 +3768,4 @@ void CSVTest()
 
 #endif // CSV_PROFILER_ALLOW_DEBUG_FEATURES
 
-#endif // CSV_PROFILER
+#endif // CSV_PROFILER
\ No newline at end of file
diff --git a/Engine/Source/Runtime/Core/Public/ProfilingDebugging/CsvProfiler.h b/Engine/Source/Runtime/Core/Public/ProfilingDebugging/CsvProfiler.h
--- a/Engine/Source/Runtime/Core/Public/ProfilingDebugging/CsvProfiler.h
+++ b/Engine/Source/Runtime/Core/Public/ProfilingDebugging/CsvProfiler.h
@@ -285,6 +285,7 @@ public:
 
 	CORE_API static bool IsWaitTrackingEnabledOnCurrentThread();
 
+	CORE_API void GetFrameExecCommands(TArray<FString>& OutFrameCommands) const;
 
 	DECLARE_MULTICAST_DELEGATE(FOnCSVProfileStart);
 	FOnCSVProfileStart& OnCSVProfileStart() { return OnCSVProfileStartDelegate; }
diff --git a/Engine/Source/Runtime/Launch/Private/LaunchEngineLoop.cpp b/Engine/Source/Runtime/Launch/Private/LaunchEngineLoop.cpp
--- a/Engine/Source/Runtime/Launch/Private/LaunchEngineLoop.cpp
+++ b/Engine/Source/Runtime/Launch/Private/LaunchEngineLoop.cpp
@@ -1370,91 +1370,68 @@ bool IsServerDelegateForOSS(FName WorldContextHandle)
 #if WITH_ENGINE && CSV_PROFILER
 static void UpdateCoreCsvStats_BeginFrame()
 {
-#if PLATFORM_WINDOWS && !UE_BUILD_SHIPPING
 	if (FCsvProfiler::Get()->IsCapturing())
 	{
-		const uint32 ProcessId = (uint32)GetCurrentProcessId();
-		float ProcessUsageFraction = 0.f, IdleUsageFraction = 0.f;
-		FWindowsPlatformProcess::GetPerFrameProcessorUsage(ProcessId, ProcessUsageFraction, IdleUsageFraction);
-
-		CSV_CUSTOM_STAT_GLOBAL(CPUUsage_Process, ProcessUsageFraction, ECsvCustomStatOp::Set);
-		CSV_CUSTOM_STAT_GLOBAL(CPUUsage_Idle, IdleUsageFraction, ECsvCustomStatOp::Set);
-	}
-#endif
-
-#if !UE_BUILD_SHIPPING
-	static TMap<uint32, TArray<FString>>* CsvFrameExecCmds = NULL;
-	if (CsvFrameExecCmds == NULL)
-	{
-		CsvFrameExecCmds = new TMap<uint32, TArray<FString>>();
-
-		FString CsvExecCommandsStr;
-		FParse::Value(FCommandLine::Get(), TEXT("-csvExecCmds="), CsvExecCommandsStr, false);
-
-		TArray<FString> CsvExecCommandsList;
-		if (CsvExecCommandsStr.ParseIntoArray(CsvExecCommandsList, TEXT(","), true) > 0)
+		if (!IsRunningDedicatedServer())
 		{
-			for (FString FrameAndCommand : CsvExecCommandsList)
-			{
-				TArray<FString> FrameAndCommandList;
-				if (FrameAndCommand.ParseIntoArray(FrameAndCommandList, TEXT(":"), true) == 2)
-				{
-					uint32 Frame = FCString::Atoi(*FrameAndCommandList[0]);
-					if (!CsvFrameExecCmds->Find(Frame))
-					{
-						CsvFrameExecCmds->Add(Frame, TArray<FString>());
-					}
-					(*CsvFrameExecCmds)[Frame].Add(FrameAndCommandList[1]);
-				}
-			}
+#if PLATFORM_WINDOWS && !UE_BUILD_SHIPPING
+		    const uint32 ProcessId = (uint32)GetCurrentProcessId();
+		    float ProcessUsageFraction = 0.f, IdleUsageFraction = 0.f;
+		    FWindowsPlatformProcess::GetPerFrameProcessorUsage(ProcessId, ProcessUsageFraction, IdleUsageFraction);
+    
+		    CSV_CUSTOM_STAT_GLOBAL(CPUUsage_Process, ProcessUsageFraction, ECsvCustomStatOp::Set);
+		    CSV_CUSTOM_STAT_GLOBAL(CPUUsage_Idle, IdleUsageFraction, ECsvCustomStatOp::Set);
+#endif
 		}
-	}
-	if (FCsvProfiler::Get()->IsCapturing())
-	{
-		TArray<FString>* FrameCommands = CsvFrameExecCmds->Find(FCsvProfiler::Get()->GetCaptureFrameNumber());
-		if (FrameCommands != NULL)
+#if CSV_PROFILER_ALLOW_DEBUG_FEATURES
+		// Handle CsvExecCmds for this frame
+		if (GEngine && GWorld)
 		{
-			for (FString Cmd : *FrameCommands)
+			TArray<FString> FrameCommands;
+			FCsvProfiler::Get()->GetFrameExecCommands(FrameCommands);
+			for (FString Cmd : FrameCommands)
 			{
-				CSV_EVENT_GLOBAL(TEXT("Executing CSV exec command : %s"), *Cmd);
+				CSV_EVENT_GLOBAL(TEXT("CsvExecCommand : %s"), *Cmd);
 				GEngine->Exec(GWorld, *Cmd);
 			}
 		}
+#endif // CSV_PROFILER_ALLOW_DEBUG_FEATURES
 	}
-
-#endif
 }
 
 static void UpdateCoreCsvStats_EndFrame()
 {
-	CSV_CUSTOM_STAT_GLOBAL(RenderThreadTime, FPlatformTime::ToMilliseconds(GRenderThreadTime), ECsvCustomStatOp::Set);
-	CSV_CUSTOM_STAT_GLOBAL(GameThreadTime, FPlatformTime::ToMilliseconds(GGameThreadTime), ECsvCustomStatOp::Set);
-	CSV_CUSTOM_STAT_GLOBAL(GPUTime, FPlatformTime::ToMilliseconds(GGPUFrameTime), ECsvCustomStatOp::Set);
-	if (IsRunningRHIInSeparateThread())
-	{
-		CSV_CUSTOM_STAT_GLOBAL(RHIThreadTime, FPlatformTime::ToMilliseconds(GRHIThreadTime), ECsvCustomStatOp::Set);
-	}
-	if (GInputLatencyTime > 0)
+	if (!IsRunningDedicatedServer())
 	{
-		CSV_CUSTOM_STAT_GLOBAL(InputLatencyTime, FPlatformTime::ToMilliseconds(GInputLatencyTime), ECsvCustomStatOp::Set);
-	}
-	FPlatformMemoryStats MemoryStats = FPlatformMemory::GetStats();
-	float PhysicalMBFree = float(MemoryStats.AvailablePhysical / 1024) / 1024.0f;
+	    CSV_CUSTOM_STAT_GLOBAL(RenderThreadTime, FPlatformTime::ToMilliseconds(GRenderThreadTime), ECsvCustomStatOp::Set);
+	    CSV_CUSTOM_STAT_GLOBAL(GameThreadTime, FPlatformTime::ToMilliseconds(GGameThreadTime), ECsvCustomStatOp::Set);
+	    CSV_CUSTOM_STAT_GLOBAL(GPUTime, FPlatformTime::ToMilliseconds(GGPUFrameTime), ECsvCustomStatOp::Set);
+	    if (IsRunningRHIInSeparateThread())
+	    {
+		    CSV_CUSTOM_STAT_GLOBAL(RHIThreadTime, FPlatformTime::ToMilliseconds(GRHIThreadTime), ECsvCustomStatOp::Set);
+	    }
+	    if (GInputLatencyTime > 0)
+	    {
+		    CSV_CUSTOM_STAT_GLOBAL(InputLatencyTime, FPlatformTime::ToMilliseconds(GInputLatencyTime), ECsvCustomStatOp::Set);
+	    }
+	    FPlatformMemoryStats MemoryStats = FPlatformMemory::GetStats();
+	    float PhysicalMBFree = float(MemoryStats.AvailablePhysical / 1024) / 1024.0f;
 #if !UE_BUILD_SHIPPING
-	// Subtract any extra development memory from physical free. This can result in negative values in cases where we would have crashed OOM
-	PhysicalMBFree -= float(FPlatformMemory::GetExtraDevelopmentMemorySize() / 1024ull / 1024ull);
+	    // Subtract any extra development memory from physical free. This can result in negative values in cases where we would have crashed OOM
+	    PhysicalMBFree -= float(FPlatformMemory::GetExtraDevelopmentMemorySize() / 1024ull / 1024ull);
 #endif
-	CSV_CUSTOM_STAT_GLOBAL(MemoryFreeMB, PhysicalMBFree, ECsvCustomStatOp::Set);
+	    CSV_CUSTOM_STAT_GLOBAL(MemoryFreeMB, PhysicalMBFree, ECsvCustomStatOp::Set);
 
 #if !UE_BUILD_SHIPPING
-	float TargetFPS = 30.0f;
-	static IConsoleVariable* MaxFPSCVar = IConsoleManager::Get().FindConsoleVariable(TEXT("t.MaxFPS"));
-	if (MaxFPSCVar && MaxFPSCVar->GetFloat() > 0)
-	{
-		TargetFPS = MaxFPSCVar->GetFloat();
-	}
-	CSV_CUSTOM_STAT_GLOBAL(MaxFrameTime, 1000.0f / TargetFPS, ECsvCustomStatOp::Set);
+	    float TargetFPS = 30.0f;
+	    static IConsoleVariable* MaxFPSCVar = IConsoleManager::Get().FindConsoleVariable(TEXT("t.MaxFPS"));
+	    if (MaxFPSCVar && MaxFPSCVar->GetFloat() > 0)
+	    {
+		    TargetFPS = MaxFPSCVar->GetFloat();
+	    }
+	    CSV_CUSTOM_STAT_GLOBAL(MaxFrameTime, 1000.0f / TargetFPS, ECsvCustomStatOp::Set);
 #endif
+	}
 }
 #endif // WITH_ENGINE && CSV_PROFILER
 
@@ -2429,11 +2406,8 @@ int32 FEngineLoop::PreInitPreStartupScreen(const TCHAR* CmdLine)
 	}
 
 #if WITH_ENGINE && CSV_PROFILER
-	if (!IsRunningDedicatedServer())
-	{
 		FCoreDelegates::OnBeginFrame.AddStatic(UpdateCoreCsvStats_BeginFrame);
 		FCoreDelegates::OnEndFrame.AddStatic(UpdateCoreCsvStats_EndFrame);
-	}
 	FCsvProfiler::Get()->Init();
 #endif
 
