Added new Core delegate event FCoreDelegates::OnAllModuleLoadingPhasesComplete
- there's currently no clean way to execute something right after all modules have loaded via their official loading phases (after OnEnginePostInit still more modules get loaded so that is too early)

diff --git a/Engine/Source/Runtime/Core/Private/Misc/CoreDelegates.cpp b/Engine/Source/Runtime/Core/Private/Misc/CoreDelegates.cpp
--- a/Engine/Source/Runtime/Core/Private/Misc/CoreDelegates.cpp
+++ b/Engine/Source/Runtime/Core/Private/Misc/CoreDelegates.cpp
@@ -90,6 +90,7 @@ FCoreDelegates::FPakSigningKeysDelegate& FCoreDelegates::GetPakSigningKeysDelega
 FSimpleMulticastDelegate FCoreDelegates::OnShutdownAfterError;
 FSimpleMulticastDelegate FCoreDelegates::OnInit;
 FSimpleMulticastDelegate FCoreDelegates::OnPostEngineInit;
+FSimpleMulticastDelegate FCoreDelegates::OnAllModuleLoadingPhasesComplete;
 FSimpleMulticastDelegate FCoreDelegates::OnFEngineLoopInitComplete;
 FSimpleMulticastDelegate FCoreDelegates::OnExit;
 FSimpleMulticastDelegate FCoreDelegates::OnPreExit;
diff --git a/Engine/Source/Runtime/Core/Public/Misc/CoreDelegates.h b/Engine/Source/Runtime/Core/Public/Misc/CoreDelegates.h
--- a/Engine/Source/Runtime/Core/Public/Misc/CoreDelegates.h
+++ b/Engine/Source/Runtime/Core/Public/Misc/CoreDelegates.h
@@ -248,6 +248,9 @@ public:
 	// Called at the end of UEngine::Init, right before loading PostEngineInit modules for both normal execution and commandlets
 	static FSimpleMulticastDelegate OnPostEngineInit;
 
+	// Called after all modules have been loaded for all phases
+	static FSimpleMulticastDelegate OnAllModuleLoadingPhasesComplete;
+
 	// Called at the very end of engine initialization, right before the engine starts ticking. This is not called for commandlets
 	static FSimpleMulticastDelegate OnFEngineLoopInitComplete;
 
diff --git a/Engine/Source/Runtime/Launch/Private/LaunchEngineLoop.cpp b/Engine/Source/Runtime/Launch/Private/LaunchEngineLoop.cpp
--- a/Engine/Source/Runtime/Launch/Private/LaunchEngineLoop.cpp
+++ b/Engine/Source/Runtime/Launch/Private/LaunchEngineLoop.cpp
@@ -3684,6 +3684,9 @@ int32 FEngineLoop::PreInitPostStartupScreen(const TCHAR* CmdLine)
 			ensure(IProjectManager::Get().LoadModulesForProject(ELoadingPhase::PostEngineInit));
 			ensure(IPluginManager::Get().LoadModulesForEnabledPlugins(ELoadingPhase::PostEngineInit));
 
+			// Call module loading phases completion callbacks
+			FCoreDelegates::OnAllModuleLoadingPhasesComplete.Broadcast();
+
 			//run automation smoke tests now that the commandlet has had a chance to override the above flags and GEngine is available
 			FAutomationTestFramework::Get().RunSmokeTests();
 
@@ -4341,6 +4344,12 @@ int32 FEngineLoop::Init()
 		}
 	}
 
+	// Call module loading phases completion callbacks
+	{
+		SCOPED_BOOT_TIMING("OnAllModuleLoadingPhasesComplete.Broadcast");
+		FCoreDelegates::OnAllModuleLoadingPhasesComplete.Broadcast();
+	}
+
 	{
 		SCOPED_BOOT_TIMING("GEngine->Start()");
 		GEngine->Start();
