Use FPlatformMisc::RequestExit when interpreting testexit

#rnx

diff --git a/Engine/Source/Runtime/Launch/Private/LaunchEngineLoop.cpp b/Engine/Source/Runtime/Launch/Private/LaunchEngineLoop.cpp
--- a/Engine/Source/Runtime/Launch/Private/LaunchEngineLoop.cpp
+++ b/Engine/Source/Runtime/Launch/Private/LaunchEngineLoop.cpp
@@ -455,21 +455,7 @@ public:
 			{
 				if (FCString::Stristr(V, *Phrase) && !FCString::Stristr(V, TEXT("-testexit=")))
 				{
-#if WITH_ENGINE
-					if (GEngine != nullptr)
-					{
-						if (GIsEditor)
-						{
-							GEngine->DeferredCommands.Add(TEXT("CLOSE_SLATE_MAINFRAME"));
-						}
-						else
-						{
-							GEngine->Exec(nullptr, TEXT("QUIT"));
-						}
-					}
-#else
 					FPlatformMisc::RequestExit(true);
-#endif
 					break;
 				}
 			}
