Enable EmitDrawEvents by default when ProfileGPU is enabled, otherwise provide options to enable except in a shipping build to avoid oversights with release builds. Also remove EmitDrawEventsOnlyOnCommandList as it is unused.

diff --git a/Engine/Source/Runtime/Core/Private/Misc/CoreGlobals.cpp b/Engine/Source/Runtime/Core/Private/Misc/CoreGlobals.cpp
--- a/Engine/Source/Runtime/Core/Private/Misc/CoreGlobals.cpp
+++ b/Engine/Source/Runtime/Core/Private/Misc/CoreGlobals.cpp
@@ -381,8 +381,6 @@ bool					GIsGameThreadIdInitialized		= false;
 void					(*GFlushStreamingFunc)(void)	  = &appNoop;
 /** Whether to emit begin/ end draw events.																	*/
 bool					GEmitDrawEvents					= false;
-/** Whether forward DrawEvents to the RHI or keep them only on the Commandlist. */
-bool					GCommandListOnlyDrawEvents		= false;
 /** Whether we want the rendering thread to be suspended, used e.g. for tracing.							*/
 bool					GShouldSuspendRenderingThread	= false;
 /** Determines what kind of trace should occur, NAME_None for none.											*/
@@ -433,21 +431,9 @@ bool GetEmitDrawEvents()
 	return GEmitDrawEvents;
 }
 
-bool CORE_API GetEmitDrawEventsOnlyOnCommandlist()
-{
-	return GCommandListOnlyDrawEvents;
-}
-
 void CORE_API SetEmitDrawEvents(bool EmitDrawEvents)
 {
 	GEmitDrawEvents = EmitDrawEvents;
-	GCommandListOnlyDrawEvents = !GEmitDrawEvents;
-}
-
-void CORE_API EnableEmitDrawEventsOnlyOnCommandlist()
-{
-	GCommandListOnlyDrawEvents = !GEmitDrawEvents;
-	GEmitDrawEvents = true;
 }
 
 void ToggleGDebugPUCrashedFlag(const TArray<FString>& Args)
diff --git a/Engine/Source/Runtime/Core/Public/CoreGlobals.h b/Engine/Source/Runtime/Core/Public/CoreGlobals.h
--- a/Engine/Source/Runtime/Core/Public/CoreGlobals.h
+++ b/Engine/Source/Runtime/Core/Public/CoreGlobals.h
@@ -726,15 +726,9 @@ struct FScopedLoadingState
 };
 #endif
 
-
 bool CORE_API GetEmitDrawEvents();
-
-bool CORE_API GetEmitDrawEventsOnlyOnCommandlist();
-
 void CORE_API SetEmitDrawEvents(bool EmitDrawEvents);
 
-void CORE_API EnableEmitDrawEventsOnlyOnCommandlist();
-
 /** Array to help visualize weak pointers in the debugger */
 class FChunkedFixedUObjectArray;
 extern CORE_API FChunkedFixedUObjectArray* GCoreObjectArrayForDebugVisualizers;
diff --git a/Engine/Source/Runtime/Engine/Private/LocalPlayer.cpp b/Engine/Source/Runtime/Engine/Private/LocalPlayer.cpp
--- a/Engine/Source/Runtime/Engine/Private/LocalPlayer.cpp
+++ b/Engine/Source/Runtime/Engine/Private/LocalPlayer.cpp
@@ -1328,7 +1328,8 @@ bool ULocalPlayer::HandleExecCommand( const TCHAR* Cmd, FOutputDevice& Ar )
 
 bool ULocalPlayer::HandleToggleDrawEventsCommand( const TCHAR* Cmd, FOutputDevice& Ar )
 {
-#if WITH_PROFILEGPU
+//Added ability to toggle this on during a test build if needed
+#if WITH_PROFILEGPU || UE_BUILD_TEST
 	if( GetEmitDrawEvents() )
 	{
 		SetEmitDrawEvents(false);
diff --git a/Engine/Source/Runtime/Launch/Private/LaunchEngineLoop.cpp b/Engine/Source/Runtime/Launch/Private/LaunchEngineLoop.cpp
--- a/Engine/Source/Runtime/Launch/Private/LaunchEngineLoop.cpp
+++ b/Engine/Source/Runtime/Launch/Private/LaunchEngineLoop.cpp
@@ -1671,6 +1671,21 @@ int32 FEngineLoop::PreInitPreStartupScreen(const TCHAR* CmdLine)
 	}
 #endif
 
+// If we are in Debug, Development, or Test/Shipping with ALLOW_PROFILEGPU... enabled, then automatically allow draw events.
+// Command lines allow disabling of draw events if WITH_PROGILEGPU is enabled. Test/Shipping on their own require explicit opt in.
+#if WITH_PROFILEGPU	
+	if (!FParse::Param(FCommandLine::Get(), TEXT("nodrawevents")))
+	{
+		SetEmitDrawEvents(true);
+	}
+// Continue to protect shipping build from emitdrawevents (if needed in shipping, ALLOW_PROFILEGPU_IN_SHIPPING should be used instead to enable WITH_PROFILEGPU)
+#elif !UE_BUILD_SHIPPING
+	if (FParse::Param(FCommandLine::Get(), TEXT("emitdrawevents")))
+	{
+		SetEmitDrawEvents(true);
+	}
+#endif
+
 #if !UE_BUILD_SHIPPING
 	if (FPlatformProperties::SupportsQuit())
 	{
@@ -1686,11 +1701,6 @@ int32 FEngineLoop::PreInitPreStartupScreen(const TCHAR* CmdLine)
 		}
 	}
 
-	if (FParse::Param(FCommandLine::Get(), TEXT("emitdrawevents")))
-	{
-		SetEmitDrawEvents(true);
-	}
-
 	// Activates malloc frame profiler from the command line 
 	// Recommend enabling bGenerateSymbols to ensure callstacks can resolve and bRetainFramePointers to ensure frame pointers remain valid.
 	// Also disabling the hitch detector ALLOW_HITCH_DETECTION=0 helps ensure quicker more accurate runs.
@@ -1701,14 +1711,6 @@ int32 FEngineLoop::PreInitPreStartupScreen(const TCHAR* CmdLine)
 	}
 #endif // !UE_BUILD_SHIPPING
 
-#if RHI_COMMAND_LIST_DEBUG_TRACES
-	// Enable command-list-only draw events if we haven't already got full draw events enabled.
-	if (!GetEmitDrawEvents())
-	{
-		EnableEmitDrawEventsOnlyOnCommandlist();
-	}
-#endif // RHI_COMMAND_LIST_DEBUG_TRACES
-
 	// Switch into executable's directory (may be required by some of the platform file overrides)
 	FPlatformProcess::SetCurrentWorkingDirectoryToBaseDir();
 
diff --git a/Engine/Source/Runtime/RHI/Public/RHICommandListCommandExecutes.inl b/Engine/Source/Runtime/RHI/Public/RHICommandListCommandExecutes.inl
--- a/Engine/Source/Runtime/RHI/Public/RHICommandListCommandExecutes.inl
+++ b/Engine/Source/Runtime/RHI/Public/RHICommandListCommandExecutes.inl
@@ -737,24 +737,12 @@ void FRHICommandEndDrawingViewport::Execute(FRHICommandListBase& CmdList)
 
 void FRHICommandPushEvent::Execute(FRHICommandListBase& CmdList)
 {
-#if	RHI_COMMAND_LIST_DEBUG_TRACES
-	if (GetEmitDrawEventsOnlyOnCommandlist())
-	{
-		return;
-	}
-#endif
 	RHISTAT(PushEvent);
 	INTERNAL_DECORATOR_COMPUTE(RHIPushEvent)(Name, Color);
 }
 
 void FRHICommandPopEvent::Execute(FRHICommandListBase& CmdList)
 {
-#if	RHI_COMMAND_LIST_DEBUG_TRACES
-	if (GetEmitDrawEventsOnlyOnCommandlist())
-	{
-		return;
-	}
-#endif
 	RHISTAT(PopEvent);
 	INTERNAL_DECORATOR_COMPUTE(RHIPopEvent)();
 }
