Core: Prevent excessive CPU use when passing -waitforattach or -waitfordebugger

diff --git a/Engine/Source/Runtime/Launch/Private/Launch.cpp b/Engine/Source/Runtime/Launch/Private/Launch.cpp
--- a/Engine/Source/Runtime/Launch/Private/Launch.cpp
+++ b/Engine/Source/Runtime/Launch/Private/Launch.cpp
@@ -99,7 +99,10 @@ int32 GuardedMain( const TCHAR* CmdLine )
 	// If "-waitforattach" or "-WaitForDebugger" was specified, halt startup and wait for a debugger to attach before continuing
 	if (FParse::Param(CmdLine, TEXT("waitforattach")) || FParse::Param(CmdLine, TEXT("WaitForDebugger")))
 	{
-		while (!FPlatformMisc::IsDebuggerPresent());
+		while (!FPlatformMisc::IsDebuggerPresent())
+		{
+			FPlatformProcess::Sleep(0.1f);
+		}
 		UE_DEBUG_BREAK();
 	}
 
