Fixed direct use of FGenericPlatformProcess rather than FPlatformProcess

#rnx

diff --git a/Engine/Source/Runtime/Launch/Private/LaunchEngineLoop.cpp b/Engine/Source/Runtime/Launch/Private/LaunchEngineLoop.cpp
--- a/Engine/Source/Runtime/Launch/Private/LaunchEngineLoop.cpp
+++ b/Engine/Source/Runtime/Launch/Private/LaunchEngineLoop.cpp
@@ -2166,7 +2166,7 @@ int32 FEngineLoop::PreInitPreStartupScreen(const TCHAR* CmdLine)
 
 #if WITH_ENGINE
 	// Add the default engine shader dir
-	AddShaderSourceDirectoryMapping(TEXT("/Engine"), FGenericPlatformProcess::ShaderDir());
+	AddShaderSourceDirectoryMapping(TEXT("/Engine"), FPlatformProcess::ShaderDir());
 
 #if WITH_EDITOR
 	{
