ZenCOTF: Make sure FStorageServerPlatformFile::InitializeAfterProjectFilePath is called before FStorageServerPlatformFile::GetFilenameOnDisk when setting the final project file path.
Required for platforms where forwarding the GetFilenameOnDisk call to the lower level filesystem will modify the path to a platform specific path.

#rnx

diff --git a/Engine/Source/Runtime/Launch/Private/LaunchEngineLoop.cpp b/Engine/Source/Runtime/Launch/Private/LaunchEngineLoop.cpp
--- a/Engine/Source/Runtime/Launch/Private/LaunchEngineLoop.cpp
+++ b/Engine/Source/Runtime/Launch/Private/LaunchEngineLoop.cpp
@@ -2121,14 +2121,16 @@ int32 FEngineLoop::PreInitPreStartupScreen(const TCHAR* CmdLine)
 		const FString ProjectFilePath = FPaths::Combine(*FPaths::ProjectDir(), *FString::Printf(TEXT("%s.%s"), FApp::GetProjectName(), *FProjectDescriptor::GetExtension()));
 		FPaths::SetProjectFilePath(ProjectFilePath);
 	}
-
-	// Fix the project file path case before we attempt to fix the game name
-	LaunchFixProjectPathCase();
 #endif
 
-	// Last chance to initialize platform file with possible knowledge of the project file path
+	// Initialize platform file with knowledge of the project file path before fixing the casing
 	FPlatformFileManager::Get().GetPlatformFile().InitializeAfterProjectFilePath();
 
+#if !IS_PROGRAM
+	// Now let the platform file fix the project file path case before we attempt to fix the game name
+	LaunchFixProjectPathCase();
+#endif
+
 	// Now verify the project file if we have one
 	if (FPaths::IsProjectFilePathSet()
 #if IS_PROGRAM
