Fixed crashes on launch of experimental device-specific mobile preview

diff --git a/Engine/Source/Editor/PIEPreviewDeviceProfileSelector/Private/PIEPreviewDeviceProfileSelectorModule.cpp b/Engine/Source/Editor/PIEPreviewDeviceProfileSelector/Private/PIEPreviewDeviceProfileSelectorModule.cpp
--- a/Engine/Source/Editor/PIEPreviewDeviceProfileSelector/Private/PIEPreviewDeviceProfileSelectorModule.cpp
+++ b/Engine/Source/Editor/PIEPreviewDeviceProfileSelector/Private/PIEPreviewDeviceProfileSelectorModule.cpp
@@ -78,7 +78,7 @@ void FPIEPreviewDeviceModule::SetPreviewDevice(const FString& DeviceName)
 	InitPreviewDevice();
 }
 
-void FPIEPreviewDeviceModule::ApplyCommandLineOverrides()
+FName FPIEPreviewDeviceModule::GetPreviewPlatformName()
 {
 	// Here we need to parse the json directly as we have not yet initialized the UObject system
 	if (JsonRootObject.IsValid())
@@ -86,9 +86,10 @@ void FPIEPreviewDeviceModule::ApplyCommandLineOverrides()
 		FString DevicePlatform;
 		if (JsonRootObject->TryGetStringField(TEXT("DevicePlatform"), DevicePlatform))
 		{
-			FCommandLine::Append(*FString::Printf(TEXT(" -ScalabilityIniPlatformOverride=%s"), *DevicePlatform));
+			return FName(*DevicePlatform);
 		}
 	}
+	return NAME_None;
 }
 
 FString const FPIEPreviewDeviceModule::GetRuntimeDeviceProfileName()
diff --git a/Engine/Source/Editor/PIEPreviewDeviceProfileSelector/Public/PIEPreviewDeviceProfileSelectorModule.h b/Engine/Source/Editor/PIEPreviewDeviceProfileSelector/Public/PIEPreviewDeviceProfileSelectorModule.h
--- a/Engine/Source/Editor/PIEPreviewDeviceProfileSelector/Public/PIEPreviewDeviceProfileSelectorModule.h
+++ b/Engine/Source/Editor/PIEPreviewDeviceProfileSelector/Public/PIEPreviewDeviceProfileSelectorModule.h
@@ -38,7 +38,7 @@ public:
 	{
 	}
 
-	virtual void ApplyCommandLineOverrides() override;
+	virtual FName GetPreviewPlatformName() override;
 
 	virtual void ApplyPreviewDeviceState() override;
 	
diff --git a/Engine/Source/Runtime/Engine/Public/IPIEPreviewDeviceModule.h b/Engine/Source/Runtime/Engine/Public/IPIEPreviewDeviceModule.h
--- a/Engine/Source/Runtime/Engine/Public/IPIEPreviewDeviceModule.h
+++ b/Engine/Source/Runtime/Engine/Public/IPIEPreviewDeviceModule.h
@@ -29,9 +29,9 @@ class IPIEPreviewDeviceModule : public IDeviceProfileSelectorModule
 		//~ Begin IPIEPreviewDeviceModule Interface
 
 		/**
-		 * Gives the PIEPreviewDeviceModule a chance to modify the command line based on the target device json
+		 * Used to set the scalability preview platform
 		 */
-		virtual void ApplyCommandLineOverrides() = 0;
+		virtual FName GetPreviewPlatformName() = 0;
 
 		/**
 		* Create PieWindow Ref
diff --git a/Engine/Source/Runtime/Launch/Private/LaunchEngineLoop.cpp b/Engine/Source/Runtime/Launch/Private/LaunchEngineLoop.cpp
--- a/Engine/Source/Runtime/Launch/Private/LaunchEngineLoop.cpp
+++ b/Engine/Source/Runtime/Launch/Private/LaunchEngineLoop.cpp
@@ -5763,18 +5763,6 @@ bool FEngineLoop::AppInit( )
 		FPlatformOutputDevices::SetupOutputDevices();
 	}
 
-#if WITH_EDITOR
-	// Append any command line overrides when running as a preview device
-	if (FPIEPreviewDeviceModule::IsRequestingPreviewDevice())
-	{
-		FPIEPreviewDeviceModule* PIEPreviewDeviceProfileSelectorModule = FModuleManager::LoadModulePtr<FPIEPreviewDeviceModule>("PIEPreviewDeviceProfileSelector");
-		if (PIEPreviewDeviceProfileSelectorModule)
-		{
-			PIEPreviewDeviceProfileSelectorModule->ApplyCommandLineOverrides();
-		}
-	}
-#endif
-
 	{
 		SCOPED_BOOT_TIMING("FConfigCacheIni::InitializeConfigSystem");
 		LLM_SCOPE(ELLMTag::ConfigSystem);
@@ -6100,6 +6088,17 @@ bool FEngineLoop::AppInit( )
 		FCoreDelegates::OnInit.Broadcast();
 	}
 
+#if WITH_EDITOR
+	if (FPIEPreviewDeviceModule::IsRequestingPreviewDevice())
+	{
+		FPIEPreviewDeviceModule* PIEPreviewDeviceProfileSelectorModule = FModuleManager::LoadModulePtr<FPIEPreviewDeviceModule>("PIEPreviewDeviceProfileSelector");
+		if (PIEPreviewDeviceProfileSelectorModule)
+		{
+			Scalability::ChangeScalabilityPreviewPlatform(PIEPreviewDeviceProfileSelectorModule->GetPreviewPlatformName());
+		}
+	}
+#endif
+
 	FEmbeddedCommunication::ForceTick(19);
 
 	return true;
