LaunchEngineLoop: Fix deadlock on exit when launching the game or editor with the parameter -testexit=ExitPhrase.
Request a deferred engine shutdown instead of forcing it from the logging code which will deadlock on most or all platforms.
-testexit is used when BuildCookRun is launched with both -run and -unattended.

#rnx

diff --git a/Engine/Source/Runtime/Launch/Private/LaunchEngineLoop.cpp b/Engine/Source/Runtime/Launch/Private/LaunchEngineLoop.cpp
--- a/Engine/Source/Runtime/Launch/Private/LaunchEngineLoop.cpp
+++ b/Engine/Source/Runtime/Launch/Private/LaunchEngineLoop.cpp
@@ -436,25 +436,31 @@ private:
 class FOutputDeviceTestExit : public FOutputDevice
 {
 	TArray<FString> ExitPhrases;
+	FString* FoundExitPhrase = nullptr;
 public:
-	FOutputDeviceTestExit(const TArray<FString>& InExitPhrases)
+	FOutputDeviceTestExit(TArray<FString>&& InExitPhrases)
 		: ExitPhrases(InExitPhrases)
 	{
 	}
 	virtual ~FOutputDeviceTestExit()
 	{
 	}
+	bool RequestExit() { return !IsEngineExitRequested() && FoundExitPhrase; }
+	FString RequestExitPhrase() { return FoundExitPhrase ? *FoundExitPhrase : FString(); }
 	virtual void Serialize(const TCHAR* V, ELogVerbosity::Type Verbosity, const class FName& Category) override
 	{
-		if (!IsEngineExitRequested())
+		if (FoundExitPhrase || IsEngineExitRequested())
 		{
-			for (auto& Phrase : ExitPhrases)
+			return;
+		}
+
+		for (auto& Phrase : ExitPhrases)
+		{
+			// look for the exit phrase, but ignore the output of the actual commandline string
+			if (FCString::Stristr(V, *Phrase) && !FCString::Stristr(V, TEXT("-testexit=")))
 			{
-				if (FCString::Stristr(V, *Phrase) && !FCString::Stristr(V, TEXT("-testexit=")))
-				{
-					FPlatformMisc::RequestExit(true);
-					break;
-				}
+				FoundExitPhrase = &Phrase;
+				break;
 			}
 		}
 	}
@@ -1697,7 +1703,7 @@ int32 FEngineLoop::PreInitPreStartupScreen(const TCHAR* CmdLine)
 			TArray<FString> ExitPhrasesList;
 			if (ExitPhrases.ParseIntoArray(ExitPhrasesList, TEXT("+"), true) > 0)
 			{
-				GScopedTestExit = MakeUnique<FOutputDeviceTestExit>(ExitPhrasesList);
+				GScopedTestExit = MakeUnique<FOutputDeviceTestExit>(MoveTemp(ExitPhrasesList));
 				GLog->AddOutputDevice(GScopedTestExit.Get());
 			}
 		}
@@ -4968,6 +4974,13 @@ void FEngineLoop::Tick()
 	LLM_SCOPE(ELLMTag::EngineMisc);
 
 	BeginExitIfRequested();
+#if !UE_BUILD_SHIPPING
+	if (GScopedTestExit.IsValid() && GScopedTestExit->RequestExit())
+	{
+		UE_LOG(LogExit, Display, TEXT("**** TestExit: %s ****"), *GScopedTestExit->RequestExitPhrase());
+		FPlatformMisc::RequestExit(true);
+	}
+#endif
 
 	// Send a heartbeat for the diagnostics thread
 	FThreadHeartBeat::Get().HeartBeat(true);
