#UE Fix missing module in commandlets when loading assets referencing WorldPartitionHLODUtilities. This is fixed in a different way in UE5Main, so this is just a temp fix until we get newer code

diff --git a/Engine/Source/Runtime/Launch/Private/LaunchEngineLoop.cpp b/Engine/Source/Runtime/Launch/Private/LaunchEngineLoop.cpp
--- a/Engine/Source/Runtime/Launch/Private/LaunchEngineLoop.cpp
+++ b/Engine/Source/Runtime/Launch/Private/LaunchEngineLoop.cpp
@@ -3981,6 +3981,8 @@ void FEngineLoop::LoadPreInitModules()
 	// Load audio editor module before engine class CDOs are loaded
 	FModuleManager::Get().LoadModule(TEXT("AudioEditor"));
 	FModuleManager::Get().LoadModule(TEXT("AnimationModifiers"));
+
+	FModuleManager::Get().LoadModule(TEXT("WorldPartitionHLODUtilities"));
 #endif	
 }
 
