Disable container overflow checks

diff --git a/Engine/Source/Runtime/Launch/Private/Unix/LaunchUnix.cpp b/Engine/Source/Runtime/Launch/Private/Unix/LaunchUnix.cpp
--- a/Engine/Source/Runtime/Launch/Private/Unix/LaunchUnix.cpp
+++ b/Engine/Source/Runtime/Launch/Private/Unix/LaunchUnix.cpp
@@ -8,16 +8,18 @@
 	#define DISABLE_ASAN_LEAK_DETECTOR 0
 #endif
 
-#if DISABLE_ASAN_LEAK_DETECTOR
 /*
  * We honestly leak so much this output is not super useful, so lets disable by default but if you want to re-enable disable
  * this DEFINE in LinuxToolchain.cs area. This must be defined in the main binary otherwise asan*.so will bind to this symbol first
  */
 extern "C" const char* LAUNCH_API __asan_default_options()
 {
-	return "detect_leaks=0";
-}
+#if DISABLE_ASAN_LEAK_DETECTOR
+	return "detect_container_overflow=0 detect_leaks=0";
+#else
+	return "detect_container_overflow=0";
 #endif
+}
 
 extern int32 GuardedMain( const TCHAR* CmdLine );
 
