Fix InputLatency csv stats on some consoles. On certain platforms, cycle is 32 bits and gets shifted to avoid precision loss. If stats is using the 64 bits version, all other conversions must use the 64 bits version.

diff --git a/Engine/Source/Runtime/Launch/Private/LaunchEngineLoop.cpp b/Engine/Source/Runtime/Launch/Private/LaunchEngineLoop.cpp
--- a/Engine/Source/Runtime/Launch/Private/LaunchEngineLoop.cpp
+++ b/Engine/Source/Runtime/Launch/Private/LaunchEngineLoop.cpp
@@ -1401,7 +1401,7 @@ static void UpdateCoreCsvStats_EndFrame()
 	    }
 	    if (GInputLatencyTime > 0)
 	    {
-		    CSV_CUSTOM_STAT_GLOBAL(InputLatencyTime, FPlatformTime::ToMilliseconds(GInputLatencyTime), ECsvCustomStatOp::Set);
+		    CSV_CUSTOM_STAT_GLOBAL(InputLatencyTime, FPlatformTime::ToMilliseconds64(GInputLatencyTime), ECsvCustomStatOp::Set);
 	    }
 	    FPlatformMemoryStats MemoryStats = FPlatformMemory::GetStats();
 	    float PhysicalMBFree = float(MemoryStats.AvailablePhysical / 1024) / 1024.0f;
