LaunchEngineLoop: Move  PRIVATE_GIsRunningCookCommandlet assignment earlier so that UE::AssetDataGather::Private::FPreloader can use it early on to start preloading.
Reduce duplication in LaunchEngineLoop of the assignment of the related variable PRIVATE_GIsRunningCommandlet, so that it is easier to verify that PRIVATE_GIsRunningCookCommandlet is set correctly.
#rnx

diff --git a/Engine/Source/Runtime/Launch/Private/LaunchEngineLoop.cpp b/Engine/Source/Runtime/Launch/Private/LaunchEngineLoop.cpp
--- a/Engine/Source/Runtime/Launch/Private/LaunchEngineLoop.cpp
+++ b/Engine/Source/Runtime/Launch/Private/LaunchEngineLoop.cpp
@@ -274,6 +274,8 @@ class FFeedbackContext;
 #include "IOS/IOSAppDelegate.h"
 #endif
 
+// MRPTODO
+PRAGMA_DISABLE_OPTIMIZATION;
 
 int32 GUseDisregardForGCOnDedicatedServers = 1;
 static FAutoConsoleVariableRef CVarUseDisregardForGCOnDedicatedServers(
@@ -1888,154 +1890,256 @@ int32 FEngineLoop::PreInitPreStartupScreen(const TCHAR* CmdLine)
 	FPlatformProcess::SetThreadAffinityMask(FPlatformAffinity::GetMainGameMask());
 	FPlatformProcess::SetupGameThread();
 
-	// Figure out whether we're the editor, ucc or the game.
+	// These bools are the mode selector and are mutually exclusive. This function selects the mode by calling one of the 
+	// SetIsRunningAs* lambdas, which it does based on commandline and configuration considerations in a specific order.
+	// TODO: Convert these bools to an enum since they are mutually exclusive.
+	bool bHasCommandletToken = false;
+	bool bHasEditorToken = false;
+	bool bIsRunningAsDedicatedServer = false;
+	bool bIsRegularClient = false;
+
 	const SIZE_T CommandLineSize = FCString::Strlen(CmdLine) + 1;
 	TCHAR* CommandLineCopy = new TCHAR[CommandLineSize];
 	FCString::Strcpy(CommandLineCopy, CommandLineSize, CmdLine);
 	const TCHAR* ParsedCmdLine = CommandLineCopy;
 
 	FString Token = FParse::Token(ParsedCmdLine, 0);
+	// trim any whitespace at edges of string - this can happen if the token was quoted with leading or trailing whitespace
+	// VC++ tends to do this in its "external tools" config
+	Token.TrimStartAndEndInline();
+
+	auto IsModeSelected = [&bHasCommandletToken, &bHasEditorToken, &bIsRegularClient, &bIsRunningAsDedicatedServer]()
+	{
+		return bHasCommandletToken || bHasEditorToken || bIsRegularClient || bIsRunningAsDedicatedServer;
+	};
+#if UE_EDITOR || WITH_ENGINE || WITH_EDITOR
+	const TCHAR* CommandletCommandLine = nullptr;
+	auto SetIsRunningAsCommandlet = [&CommandletCommandLine, &ParsedCmdLine,
+		&bHasCommandletToken, &bIsRunningAsDedicatedServer, &bHasEditorToken, &bIsRegularClient, &IsModeSelected]
+		(FStringView CommandletName)
+	{
+		checkf(!IsModeSelected(), TEXT("SetIsRunningAsCommandlet should not be called after mode has been selected."));
+		bHasCommandletToken = true;
+
+		GIsClient = true;
+		GIsServer = true;
+#if WITH_EDITORONLY_DATA
+		GIsEditor = true;
+#endif
+#if STATS
+		// Leave the stats enabled.
+		if (!FStats::EnabledForCommandlet())
+		{
+			FThreadStats::MasterDisableForever();
+		}
+#endif
 
+		CommandletCommandLine = ParsedCmdLine;
+#if WITH_EDITOR
+		if (CommandletName == TEXTVIEW("cookcommandlet"))
+		{
+			PRIVATE_GIsRunningCookCommandlet = true;
+		}
+#endif
 #if WITH_ENGINE
+		PRIVATE_GIsRunningCommandlet = true;
+#endif
+		// Allow commandlet rendering and/or audio based on command line switch (too early to let the commandlet itself override this).
+		PRIVATE_GAllowCommandletRendering = FParse::Param(FCommandLine::Get(), TEXT("AllowCommandletRendering"));
+		PRIVATE_GAllowCommandletAudio = FParse::Param(FCommandLine::Get(), TEXT("AllowCommandletAudio"));
+	};
+#endif // UE_EDITOR || WITH_ENGINE || WITH_EDITOR
+	auto SetIsRunningAsRegularClient = [&IsModeSelected, &bIsRegularClient]()
+	{
+		checkf(!IsModeSelected(), TEXT("SetIsRunningAsRegularClient should not be called after mode has been selected."));
+		bIsRegularClient = true;
 
-	TArray<FString> Tokens;
-	TArray<FString> Switches;
-	UCommandlet::ParseCommandLine(CommandLineCopy, Tokens, Switches);
+		GIsClient = true;
+		GIsServer = false;
+#if WITH_EDITORONLY_DATA
+		GIsEditor = false;
+#endif
+#if WITH_ENGINE
+		checkf(!PRIVATE_GIsRunningCommandlet, TEXT("It should not be possible for PRIVATE_GIsRunningCommandlet to have been set when calling SetIsRunningAsRegularClient"));
+#endif
+	};
+	auto SetIsRunningAsDedicatedServer = [&IsModeSelected, &bIsRunningAsDedicatedServer]()
+	{
+		checkf(!IsModeSelected(), TEXT("SetIsRunningAsDedicatedServer should not be called after mode has been selected."));
+		bIsRunningAsDedicatedServer = true;
 
-	bool bHasCommandletToken = false;
+		GIsClient = false;
+		GIsServer = true;
+#if WITH_EDITORONLY_DATA
+		GIsEditor = false;
+#endif
+#if WITH_ENGINE
+		checkf(!PRIVATE_GIsRunningCommandlet, TEXT("It should not be possible for PRIVATE_GIsRunningCommandlet to have been set when calling SetIsRunningAsDedicatedServer"));
+#endif
+	};
+#if WITH_EDITORONLY_DATA
+	auto SetIsRunningAsEditor = [&IsModeSelected, &bHasEditorToken]()
+	{
+		checkf(!IsModeSelected(), TEXT("SetIsRunningAsEditor should not be called after mode has been selected."));
+		bHasEditorToken = true;
+
+		GIsClient = true;
+		GIsServer = true;
+		GIsEditor = true;
+#if WITH_ENGINE
+		checkf(!PRIVATE_GIsRunningCommandlet, TEXT("It should not be possible for PRIVATE_GIsRunningCommandlet to have been set when calling SetIsRunningAsEditor"));
+#endif
+	};
+#endif
 
-	for (int32 TokenIndex = 0; TokenIndex < Tokens.Num(); ++TokenIndex)
+#if WITH_EDITOR
+	if (GIsUCCMakeStandaloneHeaderGenerator)
+	{
+		// Header generation is treated as a commandlet
+		checkf(!IsRunningDedicatedServer(), TEXT("GIsUCCMakeStandaloneHeaderGenerator is not supported in a ServerOnly configuration, or with -server on commandline."));
+		Token = TEXT("GIsUCCMakeStandaloneHeaderGenerator");
+		SetIsRunningAsCommandlet(Token);
+	}
+	else
+#endif
 	{
-		if (Tokens[TokenIndex].EndsWith(TEXT("Commandlet")))
+		if (IsRunningDedicatedServer())
 		{
-			bHasCommandletToken = true;
-			Token = Tokens[TokenIndex];
-			break;
+			SetIsRunningAsDedicatedServer();
 		}
 	}
 
-	for (int32 SwitchIndex = 0; SwitchIndex < Switches.Num() && !bHasCommandletToken; ++SwitchIndex)
+#if UE_EDITOR || WITH_ENGINE || WITH_EDITOR
+	TArray<FString> NonSwitchTokens;
+	TArray<FString> Switches;
+	UCommandlet::ParseCommandLine(CommandLineCopy, NonSwitchTokens, Switches);
+	if (!IsModeSelected())
 	{
-		if (Switches[SwitchIndex].StartsWith(TEXT("RUN=")))
+		for (const FString& ParsedToken : NonSwitchTokens)
+		{
+			if (ParsedToken.EndsWith(TEXT("Commandlet")))
+			{
+				Token = ParsedToken;
+				Token.TrimStartAndEndInline();
+				SetIsRunningAsCommandlet(Token);
+				break;
+			}
+		}
+		if (!bHasCommandletToken)
 		{
-			bHasCommandletToken = true;
-			Token = Switches[SwitchIndex];
-			break;
+			for (const FString& ParsedSwitch : Switches)
+			{
+				if (ParsedSwitch.StartsWith(TEXT("RUN=")))
+				{
+					Token = ParsedSwitch.RightChop(4);
+					Token.TrimStartAndEndInline();
+					if (!Token.EndsWith(TEXT("Commandlet")))
+					{
+						Token += TEXT("Commandlet");
+					}
+					SetIsRunningAsCommandlet(Token);
+					break;
+				}
+			}
 		}
 	}
+#endif // UE_EDITOR || WITH_ENGINE
 
-	if (bHasCommandletToken)
+	// In the commandlet case, we have set the Token to something other than the first token from commandline.
+	// In all other cases we should check for the first token being the Project Specifier
+	if (!bHasCommandletToken)
 	{
-		// will be reset later once the commandlet class loaded
-		PRIVATE_GIsRunningCommandlet = true;
-	}
-
-#endif // WITH_ENGINE
-
-
-	// trim any whitespace at edges of string - this can happen if the token was quoted with leading or trailing whitespace
-	// VC++ tends to do this in its "external tools" config
-	Token.TrimStartAndEndInline();
-
-	// Path returned by FPaths::GetProjectFilePath() is normalized, so may have symlinks and ~ resolved and may differ from the original path to .uproject passed in the command line
-	FString NormalizedToken = Token;
-	FPaths::NormalizeFilename(NormalizedToken);
+		// Path returned by FPaths::GetProjectFilePath() is normalized, so may have symlinks and ~ resolved and may differ from the original path to .uproject passed in the command line
+		FString NormalizedToken = Token;
+		FPaths::NormalizeFilename(NormalizedToken);
 
-	const bool bFirstTokenIsGameName = (FApp::HasProjectName() && Token == FApp::GetProjectName());
-	const bool bFirstTokenIsGameProjectFilePath = (FPaths::IsProjectFilePathSet() && NormalizedToken == FPaths::GetProjectFilePath());
-	const bool bFirstTokenIsGameProjectFileShortName = (FPaths::IsProjectFilePathSet() && Token == FPaths::GetCleanFilename(FPaths::GetProjectFilePath()));
+		const bool bFirstTokenIsGameName = (FApp::HasProjectName() && Token == FApp::GetProjectName());
+		const bool bFirstTokenIsGameProjectFilePath = (FPaths::IsProjectFilePathSet() && NormalizedToken == FPaths::GetProjectFilePath());
+		const bool bFirstTokenIsGameProjectFileShortName = (FPaths::IsProjectFilePathSet() && Token == FPaths::GetCleanFilename(FPaths::GetProjectFilePath()));
 
-	if (bFirstTokenIsGameName || bFirstTokenIsGameProjectFilePath || bFirstTokenIsGameProjectFileShortName)
-	{
-		// first item on command line was the game name, remove it in all cases
-		FString RemainingCommandline = ParsedCmdLine;
-		FCString::Strcpy(CommandLineCopy, CommandLineSize, *RemainingCommandline);
-		ParsedCmdLine = CommandLineCopy;
-
-		// Set a new command-line that doesn't include the game name as the first argument
-		FCommandLine::Set(ParsedCmdLine);
+		if (bFirstTokenIsGameName || bFirstTokenIsGameProjectFilePath || bFirstTokenIsGameProjectFileShortName)
+		{
+			// The first item on command line was the game name, ignore it and get the next item for the Token
+			FString RemainingCommandline = ParsedCmdLine;
+			FCString::Strcpy(CommandLineCopy, CommandLineSize, *RemainingCommandline);
+			ParsedCmdLine = CommandLineCopy;
 
-		Token = FParse::Token(ParsedCmdLine, 0);
-		Token.TrimStartInline();
+			// Set a new command-line that doesn't include the game name as the first argument
+			FCommandLine::Set(ParsedCmdLine);
 
-		// if the next token is a project file, then we skip it (which can happen on some platforms that combine
-		// commandlines... this handles extra .uprojects, but if you run with MyGame MyGame, we can't tell if
-		// the second MyGame is a map or not)
-		while (FPaths::GetExtension(Token) == FProjectDescriptor::GetExtension())
-		{
 			Token = FParse::Token(ParsedCmdLine, 0);
 			Token.TrimStartInline();
-		}
 
-		if (bFirstTokenIsGameProjectFilePath || bFirstTokenIsGameProjectFileShortName)
-		{
-			// Convert it to relative if possible...
-			FString RelativeGameProjectFilePath = FFileManagerGeneric::DefaultConvertToRelativePath(*FPaths::GetProjectFilePath());
-			if (RelativeGameProjectFilePath != FPaths::GetProjectFilePath())
+			// if the next token is a project file, then we skip it (which can happen on some platforms that combine
+			// commandlines... this handles extra .uprojects, but if you run with MyGame MyGame, we can't tell if
+			// the second MyGame is a map or not)
+			while (FPaths::GetExtension(Token) == FProjectDescriptor::GetExtension())
 			{
-				FPaths::SetProjectFilePath(RelativeGameProjectFilePath);
+				Token = FParse::Token(ParsedCmdLine, 0);
+				Token.TrimStartInline();
 			}
-		}
-	}
 
-	// look early for the editor token
-	bool bHasEditorToken = false;
+			if (bFirstTokenIsGameProjectFilePath || bFirstTokenIsGameProjectFileShortName)
+			{
+				// Convert it to relative if possible...
+				FString RelativeGameProjectFilePath = FFileManagerGeneric::DefaultConvertToRelativePath(*FPaths::GetProjectFilePath());
+				if (RelativeGameProjectFilePath != FPaths::GetProjectFilePath())
+				{
+					FPaths::SetProjectFilePath(RelativeGameProjectFilePath);
+				}
+			}
+		}
 
 #if UE_EDITOR
-	// Check each token for '-game', '-server' or '-run='
-	bool bIsNotEditor = false;
-
-	// This isn't necessarily pretty, but many requests have been made to allow
-	//   UnrealEditor.exe <GAMENAME> -game <map>
-	// or
-	//   UnrealEditor.exe <GAMENAME> -game 127.0.0.0
-	// We don't want to remove the -game from the commandline just yet in case
-	// we need it for something later. So, just move it to the end for now...
-	const bool bFirstTokenIsGame = (Token == TEXT("-GAME"));
-	const bool bFirstTokenIsServer = (Token == TEXT("-SERVER"));
-	const bool bFirstTokenIsModeOverride = bFirstTokenIsGame || bFirstTokenIsServer || bHasCommandletToken;
-	const TCHAR* CommandletCommandLine = nullptr;
-	if (bFirstTokenIsModeOverride)
-	{
-		bIsNotEditor = true;
-		if (bFirstTokenIsGame || bFirstTokenIsServer)
-		{
-			// Move the token to the end of the list...
+		// Handle the first token being '-game' or '-server'
+		if (Token == TEXT("-GAME") || Token == TEXT("-SERVER"))
+		{
+			// This isn't necessarily pretty, but many requests have been made to allow
+			//   UnrealEditor.exe <GAMENAME> -game <map>
+			// or
+			//   UnrealEditor.exe <GAMENAME> -game 127.0.0.0
+			// We don't want to remove the -game from the commandline just yet in case
+			// we need it for something later. So, just move it to the end for now...
 			FString RemainingCommandline = ParsedCmdLine;
 			RemainingCommandline.TrimStartInline();
 			RemainingCommandline += FString::Printf(TEXT(" %s"), *Token);
 			FCommandLine::Set(*RemainingCommandline);
 		}
-		if (bHasCommandletToken)
-		{
-#if STATS
-			// Leave the stats enabled.
-			if (!FStats::EnabledForCommandlet())
-			{
-				FThreadStats::MasterDisableForever();
-			}
 #endif
-			if (Token.StartsWith(TEXT("run=")))
-			{
-				Token.RightChopInline(4, false);
-				if (!Token.EndsWith(TEXT("Commandlet")))
-				{
-					Token += TEXT("Commandlet");
-				}
-			}
-			CommandletCommandLine = ParsedCmdLine;
+	}
+
+	// If we have no Commandlet or -server, test if we are running as the editor, and set the GIsEditor/GIsClient/GIsServer flags if so.
+	// Do this early (and certainly before AppInit) so plugin-consumed library code can know (they wouldn't otherwise)
+	// Things like the config system behave differently based on these globals, and we aren't the editor by default.
+	if (!IsModeSelected())
+	{
+#if UE_EDITOR
+		if (!Switches.Contains(TEXT("GAME")))
+		{
+			SetIsRunningAsEditor();
+		}
+#elif WITH_ENGINE && WITH_EDITOR && WITH_EDITORONLY_DATA 
+		// If a non-editor target build w/ WITH_EDITOR and WITH_EDITORONLY_DATA, use the old token check...
+		//@todo. Is this something we need to support?
+		if (Token == TEXT("EDITOR"))
+		{
+			SetIsRunningAsEditor();
 		}
+#endif
+		// Game, server and non-engine programs never run as the editor
 	}
 
-	if (bHasCommandletToken)
+#if UE_EDITOR
+	// In the UE_EDITOR configuration we now know enough to finish the mode decision and decide between commandlet or client
+	// In other configurations we still may need to check the Token for whether it is a commandlet and we make the decision below
+	if (!IsModeSelected())
 	{
-		// will be reset later once the commandlet class loaded
-		PRIVATE_GIsRunningCommandlet = true;
+		SetIsRunningAsRegularClient();
 	}
 
-	if (!bIsNotEditor && GIsGameAgnosticExe)
+	if (bHasEditorToken && GIsGameAgnosticExe)
 	{
-		// If we launched without a game name or project name, try to load the most recently loaded project file.
+		// If we launched the editor IDE (e.g. not commandlet or -game) without a game name or project name, try to load the most recently loaded project file.
 		// We can not do this if we are using a FilePlatform override since the game directory may already be established.
 		const bool bIsBuildMachine = FParse::Param(FCommandLine::Get(), TEXT("BUILDMACHINE"));
 		const bool bLoadMostRecentProjectFileIfItExists = !FApp::HasProjectName() && !bFileOverrideFound && !bIsBuildMachine && !FParse::Param(CmdLine, TEXT("norecentproject"));
@@ -2044,55 +2148,7 @@ int32 FEngineLoop::PreInitPreStartupScreen(const TCHAR* CmdLine)
 			LaunchUpdateMostRecentProjectFile();
 		}
 	}
-
-	FString CheckToken = Token;
-	bool bFoundValidToken = false;
-	while (!bFoundValidToken && (CheckToken.Len() > 0))
-	{
-		if (!bIsNotEditor)
-		{
-			bool bHasNonEditorToken = (CheckToken == TEXT("-GAME")) || (CheckToken == TEXT("-SERVER")) || (CheckToken.StartsWith(TEXT("RUN="))) || CheckToken.EndsWith(TEXT("Commandlet"));
-			if (bHasNonEditorToken)
-			{
-				bIsNotEditor = true;
-				bFoundValidToken = true;
-			}
-		}
-
-		CheckToken = FParse::Token(ParsedCmdLine, 0);
-	}
-
-	bHasEditorToken = !bIsNotEditor;
-#elif WITH_ENGINE
-	const TCHAR* CommandletCommandLine = nullptr;
-	if (bHasCommandletToken)
-	{
-#if STATS
-		// Leave the stats enabled.
-		if (!FStats::EnabledForCommandlet())
-		{
-			FThreadStats::MasterDisableForever();
-		}
 #endif
-		if (Token.StartsWith(TEXT("run=")))
-		{
-			Token.RightChopInline(4, false);
-			if (!Token.EndsWith(TEXT("Commandlet")))
-			{
-				Token += TEXT("Commandlet");
-			}
-		}
-		CommandletCommandLine = ParsedCmdLine;
-	}
-#if WITH_EDITOR && WITH_EDITORONLY_DATA
-	// If a non-editor target build w/ WITH_EDITOR and WITH_EDITORONLY_DATA, use the old token check...
-	//@todo. Is this something we need to support?
-	bHasEditorToken = Token == TEXT("EDITOR");
-#else
-	// Game, server and commandlets never set the editor token
-	bHasEditorToken = false;
-#endif
-#endif	//UE_EDITOR
 
 #if !UE_BUILD_SHIPPING
 	// Benchmarking.
@@ -2271,10 +2327,8 @@ int32 FEngineLoop::PreInitPreStartupScreen(const TCHAR* CmdLine)
 	InitializeRenderingCVarsCaching();
 #endif
 
-	bool bTokenDoesNotHaveDash = Token.Len() && FCString::Strnicmp(*Token, TEXT("-"), 1) != 0;
-
 #if WITH_EDITOR
-	// If we're running as an game but don't have a project, inform the user and exit.
+	// If we're running as a game or server but don't have a project, inform the user and exit.
 	if (bHasEditorToken == false && bHasCommandletToken == false)
 	{
 		if (!FPaths::IsProjectFilePathSet())
@@ -2285,11 +2339,6 @@ int32 FEngineLoop::PreInitPreStartupScreen(const TCHAR* CmdLine)
 		}
 	}
 
-	if (GIsUCCMakeStandaloneHeaderGenerator)
-	{
-		// Rebuilding script requires some hacks in the engine so we flag that.
-		PRIVATE_GIsRunningCommandlet = true;
-	}
 #endif //WITH_EDITOR
 
 	if (FPlatformProcess::SupportsMultithreading() && bCreateTaskGraphAndThreadPools)
@@ -2386,28 +2435,18 @@ int32 FEngineLoop::PreInitPreStartupScreen(const TCHAR* CmdLine)
 	}
 
 #if WITH_ENGINE && CSV_PROFILER
-		FCoreDelegates::OnBeginFrame.AddStatic(UpdateCoreCsvStats_BeginFrame);
-		FCoreDelegates::OnEndFrame.AddStatic(UpdateCoreCsvStats_EndFrame);
+	FCoreDelegates::OnBeginFrame.AddStatic(UpdateCoreCsvStats_BeginFrame);
+	FCoreDelegates::OnEndFrame.AddStatic(UpdateCoreCsvStats_EndFrame);
 	FCsvProfiler::Get()->Init();
 #endif
 
 #if WITH_ENGINE
 	AppLifetimeEventCapture::Init();
 
-	// Perform canonical editor check before AppInit so plugin-consumed library code can know (they wouldn't otherwise)
-	// Things like the config system behave differently based on these globals, and we aren't the editor by default
 	if (bHasEditorToken)
 	{
 #if WITH_EDITOR
-
-		// We're the editor.
-		GIsClient = true;
-		GIsServer = true;
-		GIsEditor = true;
-		PRIVATE_GIsRunningCommandlet = false;
-
 		GWarn = &UnrealEdWarn;
-
 #else //WITH_EDITOR
 		FMessageDialog::Open(EAppMsgType::Ok, NSLOCTEXT("Engine", "EditorNotSupported", "Editor not supported in this mode."));
 		FPlatformMisc::RequestExit(false);
@@ -2473,9 +2512,9 @@ PRAGMA_ENABLE_DEPRECATION_WARNINGS
 #endif
 
 #if !UE_SERVER
-		if (!IsRunningDedicatedServer())
+		if (!bIsRunningAsDedicatedServer)
 		{
-			if (!IsRunningCommandlet())
+			if (!bHasCommandletToken)
 			{
 				// Note: It is critical that resolution settings are loaded before the movie starts playing so that the window size and fullscreen state is known
 				UGameUserSettings::PreloadResolutionSettings();
@@ -2576,101 +2615,62 @@ PRAGMA_ENABLE_DEPRECATION_WARNINGS
 		FPlatformMisc::RequestExit(false);
 	}
 
-	bool bIsRegularClient = false;
-
-	if (!bHasEditorToken)
+#if !UE_EDITOR
+	// UE_EDITOR doesn't care the meaning of the token except for a few special cases (FooCommandlet, -run=, -game, -server)
+	// But when running without UE_EDITOR, we look up the token as a class to see if it is a commandlet name prefix
+	bool bIsLateCommandletToken = false;
+	if (!IsModeSelected())
 	{
-		// See whether the first token on the command line is a commandlet.
-
 		//@hack: We need to set these before calling StaticLoadClass so all required data gets loaded for the commandlets.
-		GIsClient = true;
-		GIsServer = true;
+		TGuardValue<bool> ScopedGIsClient(GIsClient, true);
+		TGuardValue<bool> ScopedGIsServer(GIsServer, true);
 #if WITH_EDITOR
-		GIsEditor = true;
-#endif	//WITH_EDITOR
-		PRIVATE_GIsRunningCommandlet = true;
-
-		// Allow commandlet rendering and/or audio based on command line switch (too early to let the commandlet itself override this).
-		PRIVATE_GAllowCommandletRendering = FParse::Param(FCommandLine::Get(), TEXT("AllowCommandletRendering"));
-		PRIVATE_GAllowCommandletAudio = FParse::Param(FCommandLine::Get(), TEXT("AllowCommandletAudio"));
+		TGuardValue<bool> ScopedGIsEditor(GIsEditor, true);
+#endif // WITH_EDITOR
+		TGuardValue<bool> ScopedGIsRunningCommandlet(PRIVATE_GIsRunningCommandlet, true);
 
-		// We need to disregard the empty token as we try finding Token + "Commandlet" which would result in finding the
-		// UCommandlet class if Token is empty.
-		bool bDefinitelyCommandlet = (bTokenDoesNotHaveDash && Token.EndsWith(TEXT("Commandlet")));
-		if (!bTokenDoesNotHaveDash)
+		bool bIsPossiblyCommandletName = Token.Len() && !Token.Contains(TEXT("-"));
+		if (bIsPossiblyCommandletName)
 		{
-			if (Token.StartsWith(TEXT("run=")))
+			UClass* TempCommandletClass = FindObject<UClass>(ANY_PACKAGE, *(Token + TEXT("Commandlet")), false);
+			if (TempCommandletClass)
 			{
-				Token.RightChopInline(4, false);
-				bDefinitelyCommandlet = true;
-				if (!Token.EndsWith(TEXT("Commandlet")))
-				{
-					Token += TEXT("Commandlet");
-				}
-			}
-		}
-		else
-		{
-			if (!bDefinitelyCommandlet)
-			{
-				UClass* TempCommandletClass = FindObject<UClass>(ANY_PACKAGE, *(Token + TEXT("Commandlet")), false);
-
-				if (TempCommandletClass)
-				{
-					check(TempCommandletClass->IsChildOf(UCommandlet::StaticClass())); // ok so you have a class that ends with commandlet that is not a commandlet
-
-					Token += TEXT("Commandlet");
-					bDefinitelyCommandlet = true;
-				}
+				checkf(TempCommandletClass->IsChildOf(UCommandlet::StaticClass()), TEXT("It is not valid to have a class that ends with \"Commandlet\" that is not a UCommandlet subclass."));
+				Token += TEXT("Commandlet");
+				bIsLateCommandletToken = true;
 			}
 		}
-
-		if (!bDefinitelyCommandlet)
-		{
-			bIsRegularClient = true;
-			GIsClient = true;
-			GIsServer = false;
-#if WITH_EDITORONLY_DATA
-			GIsEditor = false;
-#endif
-			PRIVATE_GIsRunningCommandlet = false;
-		}
+	}
+	if (bIsLateCommandletToken)
+	{
+		SetIsRunningAsCommandlet(Token);
 	}
 
-	bool bDisableDisregardForGC = bHasEditorToken;
-	if (IsRunningDedicatedServer())
+	// In non-UE_EDITOR configuration this is the point where know enough to finish the mode decision and decide between commandlet or client
+	if (!IsModeSelected())
 	{
-		GIsClient = false;
-		GIsServer = true;
-		PRIVATE_GIsRunningCommandlet = false;
-#if WITH_EDITOR
-		GIsEditor = false;
-#endif
-		bDisableDisregardForGC |= FPlatformProperties::RequiresCookedData() && (GUseDisregardForGCOnDedicatedServers == 0);
+		SetIsRunningAsRegularClient();
 	}
+#endif
 
 	// If std out device hasn't been initialized yet (there was no -stdout param in the command line) and
 	// we meet all the criteria, initialize it now.
-	if (!GScopedStdOut && !bHasEditorToken && !bIsRegularClient && !IsRunningDedicatedServer())
+	if (!GScopedStdOut && !bHasEditorToken && !bIsRegularClient && !bIsRunningAsDedicatedServer)
 	{
 		SCOPED_BOOT_TIMING("InitializeStdOutDevice");
 
 		InitializeStdOutDevice();
 	}
 
-	bool bIsCook = bHasCommandletToken && (Token == TEXT("cookcommandlet"));
 #if WITH_EDITOR
+	if (IsRunningCookCommandlet())
 	{
-		if (bIsCook)
+		ITargetPlatformManagerModule* TargetPlatformManager = GetTargetPlatformManager(false);
+		FString InitErrors;
+		if (TargetPlatformManager && TargetPlatformManager->HasInitErrors(&InitErrors))
 		{
-			// Target platform manager can only be initialized successfully after 
-			ITargetPlatformManagerModule* TargetPlatformManager = GetTargetPlatformManager(false);
-			FString InitErrors;
-			if (TargetPlatformManager && TargetPlatformManager->HasInitErrors(&InitErrors))
-			{
-				RequestEngineExit(InitErrors);
-				return 1;
-			}
+			RequestEngineExit(InitErrors);
+			return 1;
 		}
 	}
 #endif
@@ -2731,8 +2731,6 @@ PRAGMA_ENABLE_DEPRECATION_WARNINGS
 	bool bForceEnableHighDPI = false;
 #if WITH_EDITOR
 	bForceEnableHighDPI = FPIEPreviewDeviceModule::IsRequestingPreviewDevice();
-	// Update PRIVATE_GIsRunningCookCommandlet, we assume that from here, PRIVATE_GIsRunningCommandlet won't change
-	PRIVATE_GIsRunningCookCommandlet = IsRunningCommandlet() && bIsCook;
 #endif
 
 	// This must be called before any window (including the splash screen is created
@@ -3025,10 +3023,18 @@ PRAGMA_ENABLE_DEPRECATION_WARNINGS
 	PreInitContext.bWithConfigPatching = bWithConfigPatching;
 	PreInitContext.bHasEditorToken = bHasEditorToken;
 #if WITH_ENGINE
+	bool bDisableDisregardForGC = bHasEditorToken;
+	if (bIsRunningAsDedicatedServer)
+	{
+		bDisableDisregardForGC |= FPlatformProperties::RequiresCookedData() && (GUseDisregardForGCOnDedicatedServers == 0);
+	}
 	PreInitContext.bDisableDisregardForGC = bDisableDisregardForGC;
 	PreInitContext.bIsRegularClient = bIsRegularClient;
 #endif // WITH_ENGINE
-	PreInitContext.bTokenDoesNotHaveDash = bTokenDoesNotHaveDash;
+	PreInitContext.bIsPossiblyUnrecognizedCommandlet = bIsRegularClient && Token.Len() && !Token.Contains(TEXT("-"));
+	PRAGMA_DISABLE_DEPRECATION_WARNINGS;
+	PreInitContext.bTokenDoesNotHaveDash = PreInitContext.bIsPossiblyUnrecognizedCommandlet;
+	PRAGMA_ENABLE_DEPRECATION_WARNINGS;
 	PreInitContext.Token = Token;
 #if UE_EDITOR || WITH_ENGINE
 	PreInitContext.CommandletCommandLine = CommandletCommandLine;
@@ -3060,7 +3066,7 @@ int32 FEngineLoop::PreInitPostStartupScreen(const TCHAR* CmdLine)
 	bool bDisableDisregardForGC = PreInitContext.bDisableDisregardForGC;
 	bool bIsRegularClient = PreInitContext.bIsRegularClient;
 #endif // WITH_ENGINE
-	bool bTokenDoesNotHaveDash = PreInitContext.bTokenDoesNotHaveDash;
+	bool bIsPossiblyUnrecognizedCommandlet = PreInitContext.bIsPossiblyUnrecognizedCommandlet;
 	FString Token = PreInitContext.Token;
 #if UE_EDITOR || WITH_ENGINE
 	const TCHAR* CommandletCommandLine = PreInitContext.CommandletCommandLine;
@@ -3521,12 +3527,14 @@ int32 FEngineLoop::PreInitPostStartupScreen(const TCHAR* CmdLine)
 
 	SlowTask.EnterProgressFrame(5);
 
-	if (!bHasEditorToken)
+	if (!bHasEditorToken && !IsRunningDedicatedServer())
 	{
 		UClass* CommandletClass = nullptr;
 
 		if (!bIsRegularClient)
 		{
+			checkf(PRIVATE_GIsRunningCommandlet, TEXT("This should have been set in PreInitPreStartupScreen"));
+
 			CommandletClass = FindObject<UClass>(ANY_PACKAGE,*Token,false);
 			int32 PeriodIdx;
 			if (!CommandletClass && Token.FindChar('.', PeriodIdx))
@@ -3593,7 +3601,6 @@ int32 FEngineLoop::PreInitPostStartupScreen(const TCHAR* CmdLine)
 				return 1;
 			}
 #endif
-			PRIVATE_GIsRunningCommandlet = true;
 			// Reset aux log if we don't want to log to the console window.
 			if( !Default->LogToConsole )
 			{
@@ -3789,7 +3796,7 @@ int32 FEngineLoop::PreInitPostStartupScreen(const TCHAR* CmdLine)
 			// We're a regular client.
 			check(bIsRegularClient);
 
-			if (bTokenDoesNotHaveDash)
+			if (bIsPossiblyUnrecognizedCommandlet)
 			{
 				// here we give people a reasonable warning if they tried to use the short name of a commandlet
 				UClass* TempCommandletClass = FindObject<UClass>(ANY_PACKAGE,*(Token+TEXT("Commandlet")),false);
@@ -6336,3 +6343,6 @@ void FPreInitContext::Cleanup()
 }
 
 #undef LOCTEXT_NAMESPACE
+
+// MRPTODO
+PRAGMA_ENABLE_OPTIMIZATION;
diff --git a/Engine/Source/Runtime/Launch/Public/LaunchEngineLoop.h b/Engine/Source/Runtime/Launch/Public/LaunchEngineLoop.h
--- a/Engine/Source/Runtime/Launch/Public/LaunchEngineLoop.h
+++ b/Engine/Source/Runtime/Launch/Public/LaunchEngineLoop.h
@@ -24,7 +24,9 @@ struct FPreInitContext
 	bool bDisableDisregardForGC = false;
 	bool bHasEditorToken = false;
 	bool bIsRegularClient = false;
+	UE_DEPRECATED(5.1, "Use bIsPossiblyUnrecognizedCommandlet instead")
 	bool bTokenDoesNotHaveDash = false;
+	bool bIsPossiblyUnrecognizedCommandlet = false;
 
 	FString Token;
 	const TCHAR* CommandletCommandLine = nullptr;
