Remove accidentally submitted PRAGMA_DISABLE_OPTIMIZATION.
#rnx

diff --git a/Engine/Source/Runtime/Launch/Private/LaunchEngineLoop.cpp b/Engine/Source/Runtime/Launch/Private/LaunchEngineLoop.cpp
--- a/Engine/Source/Runtime/Launch/Private/LaunchEngineLoop.cpp
+++ b/Engine/Source/Runtime/Launch/Private/LaunchEngineLoop.cpp
@@ -274,8 +274,6 @@ class FFeedbackContext;
 #include "IOS/IOSAppDelegate.h"
 #endif
 
-// MRPTODO
-PRAGMA_DISABLE_OPTIMIZATION;
 
 int32 GUseDisregardForGCOnDedicatedServers = 1;
 static FAutoConsoleVariableRef CVarUseDisregardForGCOnDedicatedServers(
@@ -6343,6 +6341,3 @@ void FPreInitContext::Cleanup()
 }
 
 #undef LOCTEXT_NAMESPACE
-
-// MRPTODO
-PRAGMA_ENABLE_OPTIMIZATION;
