Restoring previous fix, now with more compiling:

Fixes for tick.DoAsyncEndOfFrameTasks.ValidateReplicatedProperties
DiffProperties skips empty layouts
FGuidReferences copy constructor to deep copy container references, needed by FFastArraySerializer properties using delta structs
Update the validation to occur around the slate tick and not during the async recording

diff --git a/Engine/Source/Runtime/Engine/Private/RepLayout.cpp b/Engine/Source/Runtime/Engine/Private/RepLayout.cpp
--- a/Engine/Source/Runtime/Engine/Private/RepLayout.cpp
+++ b/Engine/Source/Runtime/Engine/Private/RepLayout.cpp
@@ -5013,6 +5013,10 @@ bool FRepLayout::DiffProperties(
 	TConstRepDataBuffer<SourceType> Source,
 	const EDiffPropertiesFlags DiffFlags) const
 {
+	if (IsEmpty())
+	{
+		return false;
+	}
 
 	// Currently, only lifetime properties init from their defaults, so default to that,
 	// but also diff conditional properties if requested.
diff --git a/Engine/Source/Runtime/Launch/Private/LaunchEngineLoop.cpp b/Engine/Source/Runtime/Launch/Private/LaunchEngineLoop.cpp
--- a/Engine/Source/Runtime/Launch/Private/LaunchEngineLoop.cpp
+++ b/Engine/Source/Runtime/Launch/Private/LaunchEngineLoop.cpp
@@ -5319,13 +5319,7 @@ void FEngineLoop::Tick()
 		const UWorld* const GameViewportWorld = GameViewport ? GameViewport->GetWorld() : nullptr;
 		UDemoNetDriver* const CurrentDemoNetDriver = GameViewportWorld ? GameViewportWorld->GetDemoNetDriver() : nullptr;
 
-		// Optionally validate that Slate has not modified any replicated properties for client replay recording.
-		FDemoSavedPropertyState PreSlateObjectStates;
 		const bool bValidateReplicatedProperties = CurrentDemoNetDriver && CVarDoAsyncEndOfFrameTasksValidateReplicatedProperties.GetValueOnGameThread() != 0;
-		if (bValidateReplicatedProperties)
-		{
-			PreSlateObjectStates = CurrentDemoNetDriver->SavePropertyState();
-		}
 
 		if (bDoConcurrentSlateTick)
 		{
@@ -5351,8 +5345,28 @@ void FEngineLoop::Tick()
 				},
 				ConcurrentTaskCompleteEvent);
 				check(ConcurrentTask.IsValid());
+
+				// If we're validating, we want to only test the slate tick so wait for the task
+				if (bValidateReplicatedProperties)
+				{
+					CSV_SCOPED_SET_WAIT_STAT(ConcurrentWithSlate);
+
+					QUICK_SCOPE_CYCLE_COUNTER(STAT_ConcurrentWithSlateTickTasks_Wait);
+					check(ConcurrentTaskCompleteEvent);
+					ConcurrentTaskCompleteEvent->Wait();
+					FPlatformProcess::ReturnSynchEventToPool(ConcurrentTaskCompleteEvent);
+					ConcurrentTaskCompleteEvent = nullptr;
+					ConcurrentTask = nullptr;
+				}
 			}
 		}
+
+		// Optionally validate that Slate has not modified any replicated properties for client replay recording.
+		FDemoSavedPropertyState PreSlateObjectStates;
+		if (bValidateReplicatedProperties)
+		{
+			PreSlateObjectStates = CurrentDemoNetDriver->SavePropertyState();
+		}
 #endif
 
 		// Tick(Advance) Time for the application and then tick and paint slate application widgets.
diff --git a/Engine/Source/Runtime/Net/Core/Public/Net/Core/Misc/GuidReferences.h b/Engine/Source/Runtime/Net/Core/Public/Net/Core/Misc/GuidReferences.h
--- a/Engine/Source/Runtime/Net/Core/Public/Net/Core/Misc/GuidReferences.h
+++ b/Engine/Source/Runtime/Net/Core/Public/Net/Core/Misc/GuidReferences.h
@@ -70,6 +70,23 @@ public:
 
 	~FGuidReferences();
 
+	FGuidReferences& operator=(const FGuidReferences& Other) = delete;
+
+	FGuidReferences(const FGuidReferences& Other)
+		: ParentIndex(Other.ParentIndex)
+		, CmdIndex(Other.CmdIndex)
+		, NumBufferBits(Other.NumBufferBits)
+		, UnmappedGUIDs(Other.UnmappedGUIDs)
+		, MappedDynamicGUIDs(Other.MappedDynamicGUIDs)
+		, Buffer(Other.Buffer)
+		, Array(nullptr)
+	{
+		if (Other.Array)
+		{
+			Array = new FGuidReferencesMap(*Other.Array);
+		}
+	}
+
 	void CountBytes(FArchive& Ar) const
 	{
 		UnmappedGUIDs.CountBytes(Ar);
