Only make minimal RHICmdList Begin/EndFrame calls in Begin/EndFrameRenderThread when can't render

#rnx

diff --git a/Engine/Source/Runtime/Launch/Private/LaunchEngineLoop.cpp b/Engine/Source/Runtime/Launch/Private/LaunchEngineLoop.cpp
--- a/Engine/Source/Runtime/Launch/Private/LaunchEngineLoop.cpp
+++ b/Engine/Source/Runtime/Launch/Private/LaunchEngineLoop.cpp
@@ -4876,6 +4876,13 @@ static uint32 TraceFrameEventThreadId = (uint32) -1;
 
 static inline void BeginFrameRenderThread(FRHICommandListImmediate& RHICmdList, uint64 CurrentFrameCounter)
 {
+	if ( !FApp::CanEverRender() )
+	{
+		GFrameNumberRenderThread++;
+		RHICmdList.BeginFrame();
+		return;
+	}
+
 	TRACE_BEGIN_FRAME(TraceFrameType_Rendering);
 	GRHICommandList.LatchBypass();
 	GFrameNumberRenderThread++;
@@ -4928,6 +4935,12 @@ static inline void BeginFrameRenderThread(FRHICommandListImmediate& RHICmdList,
 
 static inline void EndFrameRenderThread(FRHICommandListImmediate& RHICmdList, uint64 CurrentFrameCounter)
 {
+	if ( !FApp::CanEverRender() )
+	{
+		RHICmdList.EndFrame();
+		return;
+	}
+
 #if CSV_PROFILER
 	FCsvProfiler::EndExclusiveStat("RenderThreadOther");
 #endif
