RenderingThread: Fix lack of logging the message and callstack for check() failures

### Notes
The check() message and callstack is actually printed in FWindowsErrorOutputDevice::HandleError which is called in the __except block after we RaiseException() inside ReportAssert(). However, RenderingThread's __except doesn't do this, so the contents of GErrorHist are lost to everyone except for CrashReporter. This change adds this necessary logging.

Also, add Windows Launch support for -IgnoreDebugger when choosing __try/__except or not (this precedes initialization of GIgnoreDebugger). This makes debugging easier for this class of issue.

### Testing
Repro steps are in the JIRA. Tested at that change (which has a reliable RenderingThread check() failure), and at #head where no defect exists. Tests fail and pass, respectively, as expected.

#rnx

diff --git a/Engine/Source/Runtime/Core/Private/Windows/WindowsErrorOutputDevice.cpp b/Engine/Source/Runtime/Core/Private/Windows/WindowsErrorOutputDevice.cpp
--- a/Engine/Source/Runtime/Core/Private/Windows/WindowsErrorOutputDevice.cpp
+++ b/Engine/Source/Runtime/Core/Private/Windows/WindowsErrorOutputDevice.cpp
@@ -108,7 +108,7 @@ void FWindowsErrorOutputDevice::HandleError()
 	// make sure it happens after GIsGuarded is set to 0 in case this hook crashes
 	FCoreDelegates::OnHandleSystemError.Broadcast();
 
-	// Dump the error and flush the log.
+	// Dump the error and flush the log. If you change behavior here, you should probably update RenderingThread's __except block in FRenderingThread::Run
 #if !NO_LOGGING
 	FDebug::LogFormattedMessageWithCallstack(LogWindows.GetCategoryName(), __FILE__, __LINE__, TEXT("=== Critical error: ==="), GErrorHist, ELogVerbosity::Error);
 #endif
diff --git a/Engine/Source/Runtime/Launch/Private/Windows/LaunchWindows.cpp b/Engine/Source/Runtime/Launch/Private/Windows/LaunchWindows.cpp
--- a/Engine/Source/Runtime/Launch/Private/Windows/LaunchWindows.cpp
+++ b/Engine/Source/Runtime/Launch/Private/Windows/LaunchWindows.cpp
@@ -247,6 +247,13 @@ LAUNCH_API int32 LaunchWindowsStartup( HINSTANCE hInInstance, HINSTANCE hPrevIns
 
 	bool bNoExceptionHandler = FParse::Param(CmdLine,TEXT("noexceptionhandler"));
 	(void)bNoExceptionHandler;
+
+	bool bIgnoreDebugger = FParse::Param(CmdLine, TEXT("IgnoreDebugger"));
+	(void)bIgnoreDebugger;
+
+	bool bIsDebuggerPresent = FPlatformMisc::IsDebuggerPresent() && !bIgnoreDebugger;
+	(void)bIsDebuggerPresent;
+
 	// Using the -noinnerexception parameter will disable the exception handler within native C++, which is call from managed code,
 	// which is called from this function.
 	// The default case is to have three wrapped exception handlers 
@@ -264,7 +271,7 @@ LAUNCH_API int32 LaunchWindowsStartup( HINSTANCE hInInstance, HINSTANCE hPrevIns
 #if UE_BUILD_DEBUG
 	if (GUELibraryOverrideSettings.bIsEmbedded || !GAlwaysReportCrash)
 #else
-	if (GUELibraryOverrideSettings.bIsEmbedded || bNoExceptionHandler || (FPlatformMisc::IsDebuggerPresent() && !GAlwaysReportCrash))
+	if (GUELibraryOverrideSettings.bIsEmbedded || bNoExceptionHandler || (bIsDebuggerPresent && !GAlwaysReportCrash))
 #endif
 	{
 		// Don't use exception handling when a debugger is attached to exactly trap the crash. This does NOT check
diff --git a/Engine/Source/Runtime/RenderCore/Private/RenderingThread.cpp b/Engine/Source/Runtime/RenderCore/Private/RenderingThread.cpp
--- a/Engine/Source/Runtime/RenderCore/Private/RenderingThread.cpp
+++ b/Engine/Source/Runtime/RenderCore/Private/RenderingThread.cpp
@@ -528,6 +528,12 @@ PRAGMA_ENABLE_DEPRECATION_WARNINGS
 #if !PLATFORM_SEH_EXCEPTIONS_DISABLED
 			__except(FlushRHILogsAndReportCrash(GetExceptionInformation()))
 			{
+#if !NO_LOGGING
+				// Dump the error and flush the log. This is the same logging behavior as FWindowsErrorOutputDevice::HandleError which is called in GuardedMain's caller's __except
+				FDebug::LogFormattedMessageWithCallstack(LogWindows.GetCategoryName(), __FILE__, __LINE__, TEXT("=== Critical error: ==="), GErrorHist, ELogVerbosity::Error);
+#endif
+				GLog->Panic();
+
 				GRenderingThreadError = GErrorHist;
 
 				// Use a memory barrier to ensure that the game thread sees the write to GRenderingThreadError before
