Disable protect shadow gap to prevert some third party libs from causing issues with ASan

diff --git a/Engine/Source/Runtime/Launch/Private/Unix/LaunchUnix.cpp b/Engine/Source/Runtime/Launch/Private/Unix/LaunchUnix.cpp
--- a/Engine/Source/Runtime/Launch/Private/Unix/LaunchUnix.cpp
+++ b/Engine/Source/Runtime/Launch/Private/Unix/LaunchUnix.cpp
@@ -11,13 +11,16 @@
 /*
  * We honestly leak so much this output is not super useful, so lets disable by default but if you want to re-enable disable
  * this DEFINE in LinuxToolchain.cs area. This must be defined in the main binary otherwise asan*.so will bind to this symbol first
+ *
+ * Also disable protect_shadow_gap due to a bug in some third party libs causing them to fail to init:
+ *   [REDACTED_PROFILE_URL]/sanitizers/issues/629
  */
 extern "C" const char* LAUNCH_API __asan_default_options()
 {
 #if DISABLE_ASAN_LEAK_DETECTOR
-	return "detect_container_overflow=0 detect_leaks=0";
+	return "protect_shadow_gap=0 detect_container_overflow=0 detect_leaks=0";
 #else
-	return "detect_container_overflow=0";
+	return "protect_shadow_gap=0 detect_container_overflow=0";
 #endif
 }
 
