Android hang detector and loading screen hang detector including fix for false hang trigger when app is backgrounded.

diff --git a/Engine/Source/Runtime/Core/Private/HAL/ThreadHeartBeat.cpp b/Engine/Source/Runtime/Core/Private/HAL/ThreadHeartBeat.cpp
--- a/Engine/Source/Runtime/Core/Private/HAL/ThreadHeartBeat.cpp
+++ b/Engine/Source/Runtime/Core/Private/HAL/ThreadHeartBeat.cpp
@@ -806,6 +806,12 @@ void FThreadHeartBeat::ResumeHeartBeat(bool bAllThreads)
 	bool bLastThreadResumed = false;
 	{
 		FScopeLock HeartBeatLock(&HeartBeatCritical);
+		if (GlobalSuspendCount.GetValue() == 0)
+		{
+			// Resume without matching Suspend, ignore it
+			return;
+		}
+		
 		const double CurrentTime = Clock.Seconds();
 		if (bAllThreads)
 		{
diff --git a/Engine/Source/Runtime/Core/Public/Android/AndroidPlatform.h b/Engine/Source/Runtime/Core/Public/Android/AndroidPlatform.h
--- a/Engine/Source/Runtime/Core/Public/Android/AndroidPlatform.h
+++ b/Engine/Source/Runtime/Core/Public/Android/AndroidPlatform.h
@@ -62,6 +62,7 @@ typedef FAndroidTypes FPlatformTypes;
 #define PLATFORM_SUPPORTS_VIRTUAL_TEXTURE_STREAMING	1
 #define PLATFORM_USE_ANSI_POSIX_MALLOC				1
 #define PLATFORM_SUPPORTS_LANDSCAPE_VISUAL_MESH_LOD_STREAMING 1
+#define PLATFORM_USE_MINIMAL_HANG_DETECTION			1
 #define PLATFORM_CODE_SECTION(Name)					__attribute__((section(Name)))
 
 #define PLATFORM_GLOBAL_LOG_CATEGORY				LogAndroid
diff --git a/Engine/Source/Runtime/Core/Public/Android/AndroidPlatformMisc.h b/Engine/Source/Runtime/Core/Public/Android/AndroidPlatformMisc.h
--- a/Engine/Source/Runtime/Core/Public/Android/AndroidPlatformMisc.h
+++ b/Engine/Source/Runtime/Core/Public/Android/AndroidPlatformMisc.h
@@ -59,7 +59,7 @@ public:
 
 	static bool AllowThreadHeartBeat()
 	{
-		return false;
+		return true;
 	}
 
 	struct FCPUStatTime{
diff --git a/Engine/Source/Runtime/Launch/Private/Android/LaunchAndroid.cpp b/Engine/Source/Runtime/Launch/Private/Android/LaunchAndroid.cpp
--- a/Engine/Source/Runtime/Launch/Private/Android/LaunchAndroid.cpp
+++ b/Engine/Source/Runtime/Launch/Private/Android/LaunchAndroid.cpp
@@ -30,6 +30,7 @@
 #include "HAL/PlatformFile.h"
 #include "HAL/PlatformAffinity.h"
 #include "HAL/PlatformInput.h"
+#include "HAL/ThreadHeartBeat.h"
 #include "Modules/ModuleManager.h"
 #include "IMessagingModule.h"
 #include "Android/AndroidStats.h"
@@ -1225,6 +1226,8 @@ static void ActivateApp_EventThread()
 		EventHandlerEvent->Trigger();
 	}
 
+	FThreadHeartBeat::Get().ResumeHeartBeat(true);
+
 	FPreLoadScreenManager::EnableRendering(true);
 
 	extern void AndroidThunkCpp_ShowHiddenAlertDialog();
@@ -1276,6 +1279,8 @@ static void SuspendApp_EventThread()
 
 	FPreLoadScreenManager::EnableRendering(false);
 
+	FThreadHeartBeat::Get().SuspendHeartBeat(true);
+
 	// wait for a period of time before blocking rendering
 	UE_LOG(LogAndroid, Log, TEXT("AndroidEGL::  SuspendApp_EventThread, waiting for event manager to process. tid: %d"), FPlatformTLS::GetCurrentThreadId());
 	bool bSuccess = EMDoneTrigger->Wait(4000);
