Added a new query to determine whether the install bundle manager can handle early startup patching.

 - will eventually be replaced with a better system

#rnx

diff --git a/Engine/Plugins/Experimental/DefaultInstallBundleManager/Source/Private/DefaultInstallBundleManager.cpp b/Engine/Plugins/Experimental/DefaultInstallBundleManager/Source/Private/DefaultInstallBundleManager.cpp
--- a/Engine/Plugins/Experimental/DefaultInstallBundleManager/Source/Private/DefaultInstallBundleManager.cpp
+++ b/Engine/Plugins/Experimental/DefaultInstallBundleManager/Source/Private/DefaultInstallBundleManager.cpp
@@ -3876,6 +3876,11 @@ void FDefaultInstallBundleManager::RemoveEnvironmentWantsPatchCheckBackCompatDel
 	FPatchCheck::Get().RemoveEnvironmentWantsPatchCheckBackCompatDelegate(Tag);
 }
 
+bool FDefaultInstallBundleManager::SupportsEarlyStartupPatching() const
+{
+	return false;
+}
+
 bool FDefaultInstallBundleManager::IsNullInterface() const
 {
 	return false;
diff --git a/Engine/Plugins/Experimental/DefaultInstallBundleManager/Source/Public/DefaultInstallBundleManager.h b/Engine/Plugins/Experimental/DefaultInstallBundleManager/Source/Public/DefaultInstallBundleManager.h
--- a/Engine/Plugins/Experimental/DefaultInstallBundleManager/Source/Public/DefaultInstallBundleManager.h
+++ b/Engine/Plugins/Experimental/DefaultInstallBundleManager/Source/Public/DefaultInstallBundleManager.h
@@ -657,6 +657,7 @@ public:
 	virtual void StartPatchCheck() override;
 	virtual void AddEnvironmentWantsPatchCheckBackCompatDelegate(FName Tag, FInstallBundleManagerEnvironmentWantsPatchCheck Delegate) override;
 	virtual void RemoveEnvironmentWantsPatchCheckBackCompatDelegate(FName Tag) override;
+	virtual bool SupportsEarlyStartupPatching() const override;
 
 	virtual bool IsNullInterface() const override;
 
diff --git a/Engine/Source/Runtime/InstallBundleManager/Public/InstallBundleManagerInterface.h b/Engine/Source/Runtime/InstallBundleManager/Public/InstallBundleManagerInterface.h
--- a/Engine/Source/Runtime/InstallBundleManager/Public/InstallBundleManagerInterface.h
+++ b/Engine/Source/Runtime/InstallBundleManager/Public/InstallBundleManagerInterface.h
@@ -148,6 +148,7 @@ public:
 	virtual void StartPatchCheck();
 	virtual void AddEnvironmentWantsPatchCheckBackCompatDelegate(FName Tag, FInstallBundleManagerEnvironmentWantsPatchCheck Delegate) {}
 	virtual void RemoveEnvironmentWantsPatchCheckBackCompatDelegate(FName Tag) {}
+	virtual bool SupportsEarlyStartupPatching() const = 0;
 
 	virtual bool IsNullInterface() const = 0;
 
diff --git a/Engine/Source/Runtime/Launch/Private/LaunchEngineLoop.cpp b/Engine/Source/Runtime/Launch/Private/LaunchEngineLoop.cpp
--- a/Engine/Source/Runtime/Launch/Private/LaunchEngineLoop.cpp
+++ b/Engine/Source/Runtime/Launch/Private/LaunchEngineLoop.cpp
@@ -3218,7 +3218,7 @@ int32 FEngineLoop::PreInitPostStartupScreen(const TCHAR* CmdLine)
 
 		//Now that our EarlyStartupScreen is finished, lets take the necessary steps to mount paks, apply .ini cvars, and open the shader libraries if we installed content we expect to handle
 		//If using a bundle manager, assume its handling all this stuff and that we don't have to do it.
-		if (BundleManager == nullptr || BundleManager->IsNullInterface())
+		if (BundleManager == nullptr || BundleManager->IsNullInterface() || !BundleManager->SupportsEarlyStartupPatching())
 		{
 			// Mount Paks that were installed during EarlyStartupScreen
 			if (FCoreDelegates::OnMountAllPakFiles.IsBound() && FPaths::HasProjectPersistentDownloadDir() )
diff --git a/Engine/Source/Runtime/NullInstallBundleManager/Source/NullInstallBundleManager.cpp b/Engine/Source/Runtime/NullInstallBundleManager/Source/NullInstallBundleManager.cpp
--- a/Engine/Source/Runtime/NullInstallBundleManager/Source/NullInstallBundleManager.cpp
+++ b/Engine/Source/Runtime/NullInstallBundleManager/Source/NullInstallBundleManager.cpp
@@ -135,6 +135,11 @@ class FNullInstallBundleManager : public IInstallBundleManager
 
 	}
 
+	virtual bool SupportsEarlyStartupPatching() const override
+	{
+		return false;
+	}
+
 	virtual bool IsNullInterface() const override
 	{
 		return true;
