Make sure that the virtualization system is shutdown before the modules are shutdown or we might end up with dangling pointers

#rnx

- Originally I was trying to add a notification from the default virtualization module to inform the system when it's module was being torn down and invoke the shutdown at that point, but we end up with a lot of dependencies.
- Since initialization is hard coded into the engine loop we might as well hard code the shutdown as well.

diff --git a/Engine/Source/Runtime/Launch/Private/LaunchEngineLoop.cpp b/Engine/Source/Runtime/Launch/Private/LaunchEngineLoop.cpp
--- a/Engine/Source/Runtime/Launch/Private/LaunchEngineLoop.cpp
+++ b/Engine/Source/Runtime/Launch/Private/LaunchEngineLoop.cpp
@@ -4579,6 +4579,12 @@ void FEngineLoop::Exit()
 	FIoDispatcher::Shutdown();
 #endif
 
+#if WITH_EDITOR
+	// Make sure we shut this down before the modules are torn down, we can clean this up if/when the
+	// virtualization module is moved to be a plugin
+	UE::Virtualization::Shutdown();
+#endif //WITH_EDITOR
+
 #if !PLATFORM_ANDROID // UnloadModules doesn't work on Android
 #if WITH_ENGINE
 	// Save the hot reload state
