Improved scoped cpu timers for FEngineLoop (add cpu timers for FEngineLoop::Tick, HeartBeat, TickHotfixables and TickRenderingTickables). Also added a scoped cpu timer for FStats::AdvanceFrame.

diff --git a/Engine/Source/Runtime/Core/Private/Stats/Stats2.cpp b/Engine/Source/Runtime/Core/Private/Stats/Stats2.cpp
--- a/Engine/Source/Runtime/Core/Private/Stats/Stats2.cpp
+++ b/Engine/Source/Runtime/Core/Private/Stats/Stats2.cpp
@@ -194,6 +194,7 @@ TAtomic<int32> FStats::GameThreadStatsFrame(1);
 void FStats::AdvanceFrame( bool bDiscardCallstack, const FOnAdvanceRenderingThreadStats& AdvanceRenderingThreadStatsDelegate /*= FOnAdvanceRenderingThreadStats()*/ )
 {
 #if STATS
+	TRACE_CPUPROFILER_EVENT_SCOPE(FStats::AdvanceFrame);
 	LLM_SCOPE(ELLMTag::Stats);
 	check( IsInGameThread() );
 	static int32 MasterDisableChangeTagStartFrame = -1;
diff --git a/Engine/Source/Runtime/Launch/Private/LaunchEngineLoop.cpp b/Engine/Source/Runtime/Launch/Private/LaunchEngineLoop.cpp
--- a/Engine/Source/Runtime/Launch/Private/LaunchEngineLoop.cpp
+++ b/Engine/Source/Runtime/Launch/Private/LaunchEngineLoop.cpp
@@ -184,7 +184,7 @@
 
 	#include "MoviePlayer.h"
 	#include "MoviePlayerProxy.h"
-    #include "PreLoadScreenManager.h"
+	#include "PreLoadScreenManager.h"
 	#include "InstallBundleManagerInterface.h"
 
 	#include "ShaderCodeLibrary.h"
@@ -274,7 +274,6 @@ class FFeedbackContext;
 #include "IOS/IOSAppDelegate.h"
 #endif
 
-
 int32 GUseDisregardForGCOnDedicatedServers = 1;
 static FAutoConsoleVariableRef CVarUseDisregardForGCOnDedicatedServers(
 	TEXT("gc.UseDisregardForGCOnDedicatedServers"),
@@ -1293,6 +1292,8 @@ void DeleteRecordedFileReadsFromPaks()
 #endif
 }
 
+
+
 /*-----------------------------------------------------------------------------
 	FEngineLoop implementation.
 -----------------------------------------------------------------------------*/
@@ -4987,6 +4988,7 @@ static inline void EndFrameRenderThread(FRHICommandListImmediate& RHICmdList, ui
 
 void FEngineLoop::Tick()
 {
+	TRACE_CPUPROFILER_EVENT_SCOPE(FEngineLoop::Tick);
 	SCOPE_STALL_COUNTER(FEngineLoop::Tick, 2.0);
 
 #if !UE_BUILD_SHIPPING && !UE_BUILD_TEST && MALLOC_GT_HOOKS
@@ -5006,14 +5008,23 @@ void FEngineLoop::Tick()
 	}
 #endif
 
-	// Send a heartbeat for the diagnostics thread
-	FThreadHeartBeat::Get().HeartBeat(true);
+	{
+		TRACE_CPUPROFILER_EVENT_SCOPE(HeartBeat);
+		// Send a heartbeat for the diagnostics thread
+		FThreadHeartBeat::Get().HeartBeat(true);
+	}
+
 	FGameThreadHitchHeartBeat::Get().FrameStart();
-	FPlatformMisc::TickHotfixables();
+
+	{
+		TRACE_CPUPROFILER_EVENT_SCOPE(TickHotfixables);
+		FPlatformMisc::TickHotfixables();
+	}
 
 	// Make sure something is ticking the rendering tickables in -onethread mode to avoid leaks/bugs.
 	if (!GUseThreadedRendering && !GIsRenderingThreadSuspended.Load(EMemoryOrder::Relaxed))
 	{
+		TRACE_CPUPROFILER_EVENT_SCOPE(TickRenderingTickables);
 		TickRenderingTickables();
 	}
 
@@ -5060,9 +5071,9 @@ void FEngineLoop::Tick()
 		IConsoleManager::Get().CallAllConsoleVariableSinks();
 	}
 
-	{
-		TRACE_BEGIN_FRAME(TraceFrameType_Game);
+	TRACE_BEGIN_FRAME(TraceFrameType_Game);
 
+	{
 		SCOPE_CYCLE_COUNTER(STAT_FrameTime);
 
 		#if WITH_PROFILEGPU && !UE_BUILD_SHIPPING
@@ -5554,9 +5565,10 @@ void FEngineLoop::Tick()
 #if UE_GC_TRACK_OBJ_AVAILABLE
 		SET_DWORD_STAT(STAT_Hash_NumObjects, GUObjectArray.GetObjectArrayNumMinusAvailable());
 #endif
-		TRACE_END_FRAME(TraceFrameType_Game);
 	}
 
+	TRACE_END_FRAME(TraceFrameType_Game);
+
 #if BUILD_EMBEDDED_APP
 	static double LastSleepTime = FPlatformTime::Seconds();
 	double TimeNow = FPlatformTime::Seconds();
