Use the local player to exec csvexeccmds if available, falling back to GEngine->Exec if we don't have one. This allows us to execute commands which require a player, such as HideHUDElement

diff --git a/Engine/Source/Runtime/Launch/Private/LaunchEngineLoop.cpp b/Engine/Source/Runtime/Launch/Private/LaunchEngineLoop.cpp
--- a/Engine/Source/Runtime/Launch/Private/LaunchEngineLoop.cpp
+++ b/Engine/Source/Runtime/Launch/Private/LaunchEngineLoop.cpp
@@ -1387,7 +1387,26 @@
 			for (FString Cmd : FrameCommands)
 			{
 				CSV_EVENT_GLOBAL(TEXT("CsvExecCommand : %s"), *Cmd);
-				GEngine->Exec(GWorld, *Cmd);
+
+				// Try to execute on the local player
+				bool bExecuted = false;
+				for (FConstPlayerControllerIterator Iterator = GWorld->GetPlayerControllerIterator(); Iterator; ++Iterator)
+				{
+					APlayerController* PlayerController = Iterator->Get();
+					if (PlayerController)
+					{
+						if (ULocalPlayer* LocalPlayer = Cast<ULocalPlayer>(PlayerController->Player))
+						{
+							LocalPlayer->Exec(GWorld, *Cmd, *GLog);
+							bExecuted = true;
+						}
+					}
+				}
+				if (!bExecuted)
+				{
+					// Fallback to GEngine exec
+					GEngine->Exec(GWorld, *Cmd);
+				}
 			}
 		}
 #endif // CSV_PROFILER_ALLOW_DEBUG_FEATURES
