Unnecessary u8 literals removed, which change type under C++20.

diff --git a/Engine/Source/Runtime/Launch/Private/Windows/LaunchWindows.cpp b/Engine/Source/Runtime/Launch/Private/Windows/LaunchWindows.cpp
--- a/Engine/Source/Runtime/Launch/Private/Windows/LaunchWindows.cpp
+++ b/Engine/Source/Runtime/Launch/Private/Windows/LaunchWindows.cpp
@@ -37,7 +37,7 @@ extern "C" { _declspec(dllexport) uint32 AmdPowerXpressRequestHighPerformance =
 #define USE_D3D12_REDIST (PLATFORM_DESKTOP && PLATFORM_CPU_X86_FAMILY && PLATFORM_64BITS && 1)
 #if USE_D3D12_REDIST
 extern "C" { _declspec(dllexport) extern const UINT D3D12SDKVersion = 602; }
-extern "C" { _declspec(dllexport) extern const char* D3D12SDKPath = u8".\\D3D12\\"; }
+extern "C" { _declspec(dllexport) extern const char* D3D12SDKPath = ".\\D3D12\\"; }
 #endif // USE_D3D12_REDIST
 
 /**
