Fix for sounds not being previewable in shipping editor configurations. (Audio editor asset actions not getting registered)

diff --git a/Engine/Source/Runtime/Launch/Private/LaunchEngineLoop.cpp b/Engine/Source/Runtime/Launch/Private/LaunchEngineLoop.cpp
--- a/Engine/Source/Runtime/Launch/Private/LaunchEngineLoop.cpp
+++ b/Engine/Source/Runtime/Launch/Private/LaunchEngineLoop.cpp
@@ -4008,11 +4008,14 @@ void FEngineLoop::LoadPreInitModules()
 	}
 #endif // WITH_ENGINE
 
-#if (WITH_EDITOR && !(UE_BUILD_SHIPPING || UE_BUILD_TEST))
+#if WITH_EDITOR
 	// Load audio editor module before engine class CDOs are loaded
 	FModuleManager::Get().LoadModule(TEXT("AudioEditor"));
+
+#if !(UE_BUILD_SHIPPING || UE_BUILD_TEST) // todo: revist
 	FModuleManager::Get().LoadModule(TEXT("AnimationModifiers"));
-#endif	
+#endif // !(UE_BUILD_SHIPPING || UE_BUILD_TEST)
+#endif // WITH_EDITOR
 }
 
 
@@ -4101,17 +4104,21 @@ bool FEngineLoop::LoadStartupCoreModules()
 #endif	//WITH_UNREAL_DEVELOPER_TOOLS
 
 	SlowTask.EnterProgressFrame(30);
-#if (WITH_EDITOR && !(UE_BUILD_SHIPPING || UE_BUILD_TEST))
+#if (WITH_EDITOR && !(UE_BUILD_SHIPPING || UE_BUILD_TEST)) // todo: revisit
 	// HACK: load BT editor as early as possible for statically initialized assets (non cooked BT assets needs it)
 	// cooking needs this module too
 	FModuleManager::Get().LoadModule(TEXT("BehaviorTreeEditor"));
 
 	// Ability tasks are based on GameplayTasks, so we need to make sure that module is loaded as well
 	FModuleManager::Get().LoadModule(TEXT("GameplayTasksEditor"));
+#endif
 
+#if WITH_EDITOR
 	IAudioEditorModule* AudioEditorModule = &FModuleManager::LoadModuleChecked<IAudioEditorModule>("AudioEditor");
 	AudioEditorModule->RegisterAssetActions();
+#endif // WITH_EDITOR
 
+#if (WITH_EDITOR && !(UE_BUILD_SHIPPING || UE_BUILD_TEST)) // todo: revisit
 	// Load the StringTableEditor module to register its asset actions
 	FModuleManager::Get().LoadModule("StringTableEditor");
 
