moved GThreadPool creation earlier to be right after TaskGraph initialization and before `FDelayedAutoRegisterHelper::RunAndClearDelayedAutoRegisterDelegates` call as it can launch a task that expects GThreadPool already availalble. the problem was caught by TSAN

diff --git a/Engine/Source/Runtime/Launch/Private/LaunchEngineLoop.cpp b/Engine/Source/Runtime/Launch/Private/LaunchEngineLoop.cpp
--- a/Engine/Source/Runtime/Launch/Private/LaunchEngineLoop.cpp
+++ b/Engine/Source/Runtime/Launch/Private/LaunchEngineLoop.cpp
@@ -2298,73 +2298,6 @@
 		FTaskGraphInterface::Get().AttachToThread(ENamedThreads::GameThread);
 	}
 
-	FDelayedAutoRegisterHelper::RunAndClearDelayedAutoRegisterDelegates(EDelayedRegisterRunPhase::TaskGraphSystemReady);
-
-#if STATS
-	FThreadStats::StartThread();
-#endif
-
-	FDelayedAutoRegisterHelper::RunAndClearDelayedAutoRegisterDelegates(EDelayedRegisterRunPhase::StatSystemReady);
-
-	FScopeCycleCounter CycleCount_AfterStats(GET_STATID(STAT_FEngineLoop_PreInitPreStartupScreen_AfterStats));
-
-	// Load Core modules required for everything else to work (needs to be loaded before InitializeRenderingCVarsCaching)
-	{
-		SCOPED_BOOT_TIMING("LoadCoreModules");
-		if (!LoadCoreModules())
-		{
-			UE_LOG(LogInit, Error, TEXT("Failed to load Core modules."));
-			return 1;
-		}
-	}
-
-	const bool bDumpEarlyConfigReads = FParse::Param(FCommandLine::Get(), TEXT("DumpEarlyConfigReads"));
-	const bool bDumpEarlyPakFileReads = FParse::Param(FCommandLine::Get(), TEXT("DumpEarlyPakFileReads"));
-	const bool bForceQuitAfterEarlyReads = FParse::Param(FCommandLine::Get(), TEXT("ForceQuitAfterEarlyReads"));
-
-	// Overly verbose to avoid a dumb static analysis warning
-#if WITH_CONFIG_PATCHING
-	constexpr bool bWithConfigPatching = true;
-#else
-	constexpr bool bWithConfigPatching = false;
-#endif
-
-	if (bDumpEarlyConfigReads)
-	{
-		UE::ConfigUtilities::RecordConfigReadsFromIni();
-	}
-
-	if (bDumpEarlyPakFileReads)
-	{
-		RecordFileReadsFromPaks();
-	}
-
-	if(bWithConfigPatching)
-	{
-		UE_LOG(LogInit, Verbose, TEXT("Begin recording CVar changes for config patching."));
-
-		UE::ConfigUtilities::RecordApplyCVarSettingsFromIni();
-	}
-
-#if WITH_ENGINE
-	extern ENGINE_API void InitializeRenderingCVarsCaching();
-	InitializeRenderingCVarsCaching();
-#endif
-
-#if WITH_EDITOR
-	// If we're running as a game or server but don't have a project, inform the user and exit.
-	if (bHasEditorToken == false && bHasCommandletToken == false)
-	{
-		if (!FPaths::IsProjectFilePathSet())
-		{
-			//@todo this is too early to localize
-			FMessageDialog::Open(EAppMsgType::Ok, NSLOCTEXT("Engine", "UERequiresProjectFiles", "Unreal Engine games require a project file as the first parameter."));
-			return 1;
-		}
-	}
-
-#endif //WITH_EDITOR
-
 	if (FPlatformProcess::SupportsMultithreading() && bCreateTaskGraphAndThreadPools)
 	{
 		SCOPED_BOOT_TIMING("Init FQueuedThreadPool's");
@@ -2442,9 +2375,76 @@
 
 			verify(GBackgroundPriorityThreadPool->Create(NumThreadsInThreadPool, StackSize, TPri_Lowest, TEXT("BackgroundThreadPool")));
 		}
-
 	}
 
+	// this can start using TaskGraph and ThreadPool so they must be created before
+	FDelayedAutoRegisterHelper::RunAndClearDelayedAutoRegisterDelegates(EDelayedRegisterRunPhase::TaskGraphSystemReady);
+
+#if STATS
+	FThreadStats::StartThread();
+#endif
+
+	FDelayedAutoRegisterHelper::RunAndClearDelayedAutoRegisterDelegates(EDelayedRegisterRunPhase::StatSystemReady);
+
+	FScopeCycleCounter CycleCount_AfterStats(GET_STATID(STAT_FEngineLoop_PreInitPreStartupScreen_AfterStats));
+
+	// Load Core modules required for everything else to work (needs to be loaded before InitializeRenderingCVarsCaching)
+	{
+		SCOPED_BOOT_TIMING("LoadCoreModules");
+		if (!LoadCoreModules())
+		{
+			UE_LOG(LogInit, Error, TEXT("Failed to load Core modules."));
+			return 1;
+		}
+	}
+
+	const bool bDumpEarlyConfigReads = FParse::Param(FCommandLine::Get(), TEXT("DumpEarlyConfigReads"));
+	const bool bDumpEarlyPakFileReads = FParse::Param(FCommandLine::Get(), TEXT("DumpEarlyPakFileReads"));
+	const bool bForceQuitAfterEarlyReads = FParse::Param(FCommandLine::Get(), TEXT("ForceQuitAfterEarlyReads"));
+
+	// Overly verbose to avoid a dumb static analysis warning
+#if WITH_CONFIG_PATCHING
+	constexpr bool bWithConfigPatching = true;
+#else
+	constexpr bool bWithConfigPatching = false;
+#endif
+
+	if (bDumpEarlyConfigReads)
+	{
+		UE::ConfigUtilities::RecordConfigReadsFromIni();
+	}
+
+	if (bDumpEarlyPakFileReads)
+	{
+		RecordFileReadsFromPaks();
+	}
+
+	if(bWithConfigPatching)
+	{
+		UE_LOG(LogInit, Verbose, TEXT("Begin recording CVar changes for config patching."));
+
+		UE::ConfigUtilities::RecordApplyCVarSettingsFromIni();
+	}
+
+#if WITH_ENGINE
+	extern ENGINE_API void InitializeRenderingCVarsCaching();
+	InitializeRenderingCVarsCaching();
+#endif
+
+#if WITH_EDITOR
+	// If we're running as a game or server but don't have a project, inform the user and exit.
+	if (bHasEditorToken == false && bHasCommandletToken == false)
+	{
+		if (!FPaths::IsProjectFilePathSet())
+		{
+			//@todo this is too early to localize
+			FMessageDialog::Open(EAppMsgType::Ok, NSLOCTEXT("Engine", "UERequiresProjectFiles", "Unreal Engine games require a project file as the first parameter."));
+			return 1;
+		}
+	}
+
+#endif //WITH_EDITOR
+
 #if WITH_APPLICATION_CORE
 	// Get a pointer to the log output device
 	GLogConsole = GScopedLogConsole.Get();
