Android compatibility: add an allowlist of cvars which can be set via configrules variables.

diff --git a/Engine/Source/Runtime/Core/Private/Android/AndroidPlatformMisc.cpp b/Engine/Source/Runtime/Core/Private/Android/AndroidPlatformMisc.cpp
--- a/Engine/Source/Runtime/Core/Private/Android/AndroidPlatformMisc.cpp
+++ b/Engine/Source/Runtime/Core/Private/Android/AndroidPlatformMisc.cpp
@@ -2403,7 +2403,7 @@ FString FAndroidMisc::GetVulkanVersion()
 	return VulkanVersionString;
 }
 
-TMap<FString, FString> FAndroidMisc::GetConfigRulesTMap()
+const TMap<FString, FString>& FAndroidMisc::GetConfigRulesTMap()
 {
 	return ConfigRulesVariables;
 }
diff --git a/Engine/Source/Runtime/Core/Public/Android/AndroidPlatformMisc.h b/Engine/Source/Runtime/Core/Public/Android/AndroidPlatformMisc.h
--- a/Engine/Source/Runtime/Core/Public/Android/AndroidPlatformMisc.h
+++ b/Engine/Source/Runtime/Core/Public/Android/AndroidPlatformMisc.h
@@ -184,7 +184,7 @@ public:
 #endif
 	static bool IsSupportedAndroidDevice();
 	static void SetForceUnsupported(bool bInOverride);
-	static TMap<FString, FString> GetConfigRulesTMap();
+	static const TMap<FString, FString>& GetConfigRulesTMap();
 	static FString* GetConfigRulesVariable(const FString& Key);
 
 	/* HasVulkanDriverSupport
diff --git a/Engine/Source/Runtime/Launch/Private/Android/LaunchAndroid.cpp b/Engine/Source/Runtime/Launch/Private/Android/LaunchAndroid.cpp
--- a/Engine/Source/Runtime/Launch/Private/Android/LaunchAndroid.cpp
+++ b/Engine/Source/Runtime/Launch/Private/Android/LaunchAndroid.cpp
@@ -363,6 +363,36 @@ static void OnNativeWindowResized(ANativeActivity* activity, ANativeWindow* wind
 	write(app->msgwrite, &cmd, sizeof(cmd));
 }
 
+static void ApplyAndroidCompatConfigRules()
+{
+	TArray<FString> AndroidCompatCVars;
+	if (GConfig->GetArray(TEXT("AndroidCompatCVars"), TEXT("CVars"), AndroidCompatCVars, GEngineIni))
+	{
+		TSet<FString> AllowedCompatCVars(AndroidCompatCVars);
+		for (const TTuple<FString, FString>& Pair : FAndroidMisc::GetConfigRulesTMap())
+		{
+			const FString& Key = Pair.Key;
+			const FString& Value = Pair.Value;
+			static const TCHAR AndroidCompat[] = TEXT("AndroidCompat.");
+			if (Key.StartsWith(AndroidCompat))
+			{
+				FString CVarName = Key.Mid(UE_ARRAY_COUNT(AndroidCompat)-1);
+				if (AllowedCompatCVars.Contains(CVarName))
+				{
+					auto* CVar = IConsoleManager::Get().FindConsoleVariable(*CVarName);
+					if (CVar)
+					{
+						// set with current priority means that DPs etc can still override anything set here.
+						// e.g. -dpcvars= is expected to work.
+						CVar->SetWithCurrentPriority(*Value); 
+						UE_LOG(LogAndroid, Log, TEXT("Compat Setting %s = %s"), *CVarName, *Value);
+					}
+				}
+			}
+		}
+	}
+}
+
 //Main function called from the android entry point
 int32 AndroidMain(struct android_app* state)
 {
@@ -487,9 +517,13 @@ int32 AndroidMain(struct android_app* state)
 		GAndroidWindowLock.Lock();
 	}
 
+	FDelegateHandle ConfigReadyHandle = FCoreDelegates::ConfigReadyForUse.AddStatic(&ApplyAndroidCompatConfigRules);
+
 	// initialize the engine
 	int32 PreInitResult = GEngineLoop.PreInit(0, NULL, FCommandLine::Get());
 
+	FCoreDelegates::ConfigReadyForUse.Remove(ConfigReadyHandle);
+
 	if (PreInitResult != 0)
 	{
 		checkf(false, TEXT("Engine Preinit Failed"));
