Fix for LaunchFixProjectPathCase() assuming that TCHARs are WIDECHARs and doing in-string replacement with them.
Fix for FOutputDeviceStdOutput::Seralize() of TCHARs as UTF-8.

diff --git a/Engine/Source/Runtime/Launch/Private/LaunchEngineLoop.cpp b/Engine/Source/Runtime/Launch/Private/LaunchEngineLoop.cpp
--- a/Engine/Source/Runtime/Launch/Private/LaunchEngineLoop.cpp
+++ b/Engine/Source/Runtime/Launch/Private/LaunchEngineLoop.cpp
@@ -401,12 +401,15 @@ public:
 			// fall through to standard printf path
 #endif
 
-#if PLATFORM_TCHAR_IS_CHAR16
-			printf("%s\n", TCHAR_TO_UTF8(*line));
+#if PLATFORM_TCHAR_IS_UTF8CHAR
+			printf("%s\n", (const char*)*line);
+#elif PLATFORM_TCHAR_IS_CHAR16
+			printf("%s\n", (const char*)StringCast<UTF8CHAR>(*line).Get());
 #elif PLATFORM_USE_LS_SPEC_FOR_WIDECHAR
 			// printf prints wchar_t strings just fine with %ls, while mixing printf()/wprintf() is not recommended (see https://stackoverflow.com/questions/8681623/printf-and-wprintf-in-single-c-code)
 			printf("%ls\n", *line);
 #else
+			static_assert(std::is_same_v<TCHAR, WIDECHAR>, "Assuming wide chars here");
 			wprintf(TEXT("%s\n"), *line);
 #endif
 
@@ -689,34 +692,38 @@ void LaunchFixProjectPathCase()
 		// for workflows that use a junction at their root (eg: p4 gets confused about paths and operations fail).
 		// There is a way to get a case-accurate path on Windows without resolving directory junctions, but it is slow.
 		// We can use it here for this one-off situation without causing all uses of GetFilenameOnDisk to be slower.
-		WIN32_FIND_DATAW Data;
 		FString ProjectFilePath = FPaths::GetProjectFilePath();
-		for (FStringView PendingFix = ProjectFilePath; PendingFix.Len() > 0;)
-		{
-			const FStringView CurrentPathComponent = FPathViews::GetCleanFilename(PendingFix);
-
-			// Skip over all segments that are either empty or contain relative transforms or start with the volume separator, they should remain as-is
-			const bool bIsIgnoredSegment = CurrentPathComponent.IsEmpty() || CurrentPathComponent.Equals(TEXTVIEW(".")) || CurrentPathComponent.Equals(TEXTVIEW("..")) || CurrentPathComponent.EndsWith(TEXT(':'));
-			if (!bIsIgnoredSegment)
+		TStringBuilder<MAX_PATH> Builder;
+		FPathViews::IterateComponents(
+			ProjectFilePath,
+			[&Builder](FStringView CurrentPathComponent)
 			{
-				// Temporarily null-terminate the current path component for the system call
-				TCHAR Terminator = TEXT('\0');
-				Swap(ProjectFilePath.GetCharArray()[PendingFix.Len()], Terminator);
+				if (Builder.Len() != 0)
+				{
+					Builder.AppendChar(TEXT('/'));
+				}
 
-				HANDLE Handle = FindFirstFileW(*ProjectFilePath, &Data);
-				if (Handle != INVALID_HANDLE_VALUE)
+				int32 LenBeforeCurrentComponent = Builder.Len();
+
+				Builder.Append(CurrentPathComponent);
+
+				// Skip over all segments that are either empty or contain relative transforms or start with the volume separator, they should remain as-is
+				const bool bIsIgnoredSegment = CurrentPathComponent.IsEmpty() || CurrentPathComponent.Equals(TEXTVIEW(".")) || CurrentPathComponent.Equals(TEXTVIEW("..")) || CurrentPathComponent.EndsWith(TEXT(':'));
+				if (bIsIgnoredSegment)
 				{
-					check(CurrentPathComponent.Equals(Data.cFileName, ESearchCase::IgnoreCase));
-					FCString::Strncpy(GetData(ProjectFilePath) + PendingFix.Len() - CurrentPathComponent.Len(), Data.cFileName, CurrentPathComponent.Len() + 1);
-					FindClose(Handle);
+					return;
 				}
 
-				Swap(ProjectFilePath.GetCharArray()[PendingFix.Len()], Terminator);
+				WIN32_FIND_DATAW Data;
+				HANDLE Handle = FindFirstFileW(StringCast<WIDECHAR>(*Builder, Builder.Len() + 1).Get(), &Data);
+				if (Handle != INVALID_HANDLE_VALUE)
+				{
+					Builder.RemoveSuffix(Builder.Len() - LenBeforeCurrentComponent);
+					Builder.Append(Data.cFileName);
+				}
 			}
-
-			PendingFix.LeftChopInline(CurrentPathComponent.Len() + 1);
-		}
-		FPaths::SetProjectFilePath(ProjectFilePath);
+		);
+		FPaths::SetProjectFilePath(Builder.ToString());
 #else
 		FPaths::SetProjectFilePath(IFileManager::Get().GetFilenameOnDisk(*FPaths::GetProjectFilePath()));
 #endif
