Fix PVS Studio warning (resource leak).

diff --git a/Engine/Source/Runtime/Launch/Private/LaunchEngineLoop.cpp b/Engine/Source/Runtime/Launch/Private/LaunchEngineLoop.cpp
--- a/Engine/Source/Runtime/Launch/Private/LaunchEngineLoop.cpp
+++ b/Engine/Source/Runtime/Launch/Private/LaunchEngineLoop.cpp
@@ -720,6 +720,7 @@ void LaunchFixProjectPathCase()
 				{
 					Builder.RemoveSuffix(Builder.Len() - LenBeforeCurrentComponent);
 					Builder.Append(Data.cFileName);
+					FindClose(Handle);
 				}
 			}
 		);
