Load TimeManagement module on editor startup so as it is required by some serialization functions (i.e. AnimSequence)

diff --git a/Engine/Source/Runtime/Launch/Private/LaunchEngineLoop.cpp b/Engine/Source/Runtime/Launch/Private/LaunchEngineLoop.cpp
--- a/Engine/Source/Runtime/Launch/Private/LaunchEngineLoop.cpp
+++ b/Engine/Source/Runtime/Launch/Private/LaunchEngineLoop.cpp
@@ -4159,6 +4159,10 @@ bool FEngineLoop::LoadStartupCoreModules()
 #if WITH_EDITOR
 	FModuleManager::Get().LoadModule(TEXT("ClothingSystemEditor"));
 	FModuleManager::Get().LoadModule(TEXT("AnimationDataController"));
+
+	// Required during serialization of AnimSequence which could happen from async loading thread.
+	// See UAnimSequence::UpdateFrameRate().
+	FModuleManager::Get().LoadModule(TEXT("TimeManagement"));
 #endif
 
 	FModuleManager::Get().LoadModule(TEXT("PacketHandler"));
