Fix stretched viewport on some devices when returning from background with an orientation change.

diff --git a/Engine/Source/Runtime/Launch/Private/Android/LaunchAndroid.cpp b/Engine/Source/Runtime/Launch/Private/Android/LaunchAndroid.cpp
--- a/Engine/Source/Runtime/Launch/Private/Android/LaunchAndroid.cpp
+++ b/Engine/Source/Runtime/Launch/Private/Android/LaunchAndroid.cpp
@@ -131,13 +131,13 @@ extern "C"
 }
 #endif
 
-int32 GAndroidEnableNativeResizeEvent = 0;
+int32 GAndroidEnableNativeResizeEvent = 1;
 static FAutoConsoleVariableRef CVarEnableResizeNativeEvent(
 	TEXT("Android.EnableNativeResizeEvent"),
 	GAndroidEnableNativeResizeEvent,
 	TEXT("Whether native resize event is enabled on Android.\n")
-	TEXT(" 0: disabled (default)\n")
-	TEXT(" 1: enabled"),
+	TEXT(" 0: disabled\n")
+	TEXT(" 1: enabled (default)"),
 	ECVF_ReadOnly);
 
 int32 GAndroidEnableMouse = 0;
