Allow GIsFirstInstance in shipping editor

diff --git a/Engine/Plugins/Online/Experimental/OnlineServicesNull/Source/Private/Online/AuthNull.cpp b/Engine/Plugins/Online/Experimental/OnlineServicesNull/Source/Private/Online/AuthNull.cpp
--- a/Engine/Plugins/Online/Experimental/OnlineServicesNull/Source/Private/Online/AuthNull.cpp
+++ b/Engine/Plugins/Online/Experimental/OnlineServicesNull/Source/Private/Online/AuthNull.cpp
@@ -56,13 +56,11 @@ FString GenerateRandomUserId(const FAuthNullConfig& Config, FPlatformUserId Plat
 		UserSuffix = FString::Printf(TEXT("-%d"), PlatformUserId.GetInternalId());
 	}
 	 
-#if !(UE_BUILD_SHIPPING && WITH_EDITOR)
-	if (GIsFirstInstance && !GIsEditor)
+	if (FPlatformProcess::IsFirstInstance() && !GIsEditor)
 	{
 		// If we're outside the editor and know this is the first instance, use the system login id
 		bUseStableNullId = true;
 	}
-#endif
 
 	if (bUseStableNullId)
 	{
diff --git a/Engine/Plugins/Online/OnlineSubsystemNull/Source/Private/OnlineIdentityNull.cpp b/Engine/Plugins/Online/OnlineSubsystemNull/Source/Private/OnlineIdentityNull.cpp
--- a/Engine/Plugins/Online/OnlineSubsystemNull/Source/Private/OnlineIdentityNull.cpp
+++ b/Engine/Plugins/Online/OnlineSubsystemNull/Source/Private/OnlineIdentityNull.cpp
@@ -62,13 +62,11 @@ FString FOnlineIdentityNull::GenerateRandomUserId(int32 LocalUserNum)
 		UserSuffix = FString::Printf(TEXT("-%d"), LocalUserNum);
 	}
 
-#if !(UE_BUILD_SHIPPING && WITH_EDITOR)
-	if (GIsFirstInstance && !GIsEditor)
+	if (FPlatformProcess::IsFirstInstance() && !GIsEditor)
 	{
 		// If we're outside the editor and know this is the first instance, use the system login id
 		bUseStableNullId = true;
 	}
-#endif
 
 	if (bUseStableNullId)
 	{
diff --git a/Engine/Plugins/XGEController/Source/Private/XGEControllerModule.cpp b/Engine/Plugins/XGEController/Source/Private/XGEControllerModule.cpp
--- a/Engine/Plugins/XGEController/Source/Private/XGEControllerModule.cpp
+++ b/Engine/Plugins/XGEController/Source/Private/XGEControllerModule.cpp
@@ -268,18 +268,13 @@ bool FXGEControllerModule::IsSupported()
 
 void FXGEControllerModule::CleanWorkingDirectory()
 {
-#if UE_BUILD_SHIPPING && WITH_EDITOR
-	// this is not expected to run in a shipped editor, but if it is, we should deal with it
-	unimplemented();
-#else
 	// Only clean the directory if we are the only instance running,
 	// and we're not running in multi-process mode.
-	if ((GIsFirstInstance) && !FParse::Param(FCommandLine::Get(), TEXT("Multiprocess")))
+	if (FPlatformProcess::IsFirstInstance() && !FParse::Param(FCommandLine::Get(), TEXT("Multiprocess")))
 	{
 		UE_LOG(LogXGEController, Log, TEXT("Cleaning working directory: %s"), *RootWorkingDirectory);
 		IFileManager::Get().DeleteDirectory(*RootWorkingDirectory, false, true);
 	}
-#endif
 }
 
 void FXGEControllerModule::StartupModule()
diff --git a/Engine/Source/Runtime/Core/Private/GenericPlatform/GenericPlatformProcess.cpp b/Engine/Source/Runtime/Core/Private/GenericPlatform/GenericPlatformProcess.cpp
--- a/Engine/Source/Runtime/Core/Private/GenericPlatform/GenericPlatformProcess.cpp
+++ b/Engine/Source/Runtime/Core/Private/GenericPlatform/GenericPlatformProcess.cpp
@@ -647,11 +647,7 @@ bool FGenericPlatformProcess::Daemonize()
 
 bool FGenericPlatformProcess::IsFirstInstance()
 {
-#if !(UE_BUILD_SHIPPING && WITH_EDITOR)
-	return GIsFirstInstance;
-#else
 	return true;
-#endif
 }
 
 FSystemWideCriticalSectionNotImplemented::FSystemWideCriticalSectionNotImplemented(const FString& Name, FTimespan Timeout)
diff --git a/Engine/Source/Runtime/Core/Private/Misc/CoreGlobals.cpp b/Engine/Source/Runtime/Core/Private/Misc/CoreGlobals.cpp
--- a/Engine/Source/Runtime/Core/Private/Misc/CoreGlobals.cpp
+++ b/Engine/Source/Runtime/Core/Private/Misc/CoreGlobals.cpp
@@ -359,13 +359,6 @@ uint64					GInputTime					= 0;
 uint32					GFrameNumber					= 1;
 /** NEED TO RENAME, for RT version of GFrameTime use View.ViewFamily->FrameNumber or pass down from RT from GFrameTime). */
 uint32					GFrameNumberRenderThread		= 1;
-#if !(UE_BUILD_SHIPPING && WITH_EDITOR)
-// We cannot count on this variable to be accurate in a shipped game, so make sure no code tries to use it
-/** Whether we are the first instance of the game running.													*/
-#if !PLATFORM_UNIX
-bool					GIsFirstInstance				= true;
-#endif
-#endif
 /** Threshold for a frame to be considered a hitch (in milliseconds). */
 float GHitchThresholdMS = 60.0f;
 /** Size to break up data into when saving compressed data													*/
diff --git a/Engine/Source/Runtime/Core/Private/Unix/UnixPlatformProcess.cpp b/Engine/Source/Runtime/Core/Private/Unix/UnixPlatformProcess.cpp
--- a/Engine/Source/Runtime/Core/Private/Unix/UnixPlatformProcess.cpp
+++ b/Engine/Source/Runtime/Core/Private/Unix/UnixPlatformProcess.cpp
@@ -2005,7 +2005,6 @@ bool FUnixPlatformProcess::IsFirstInstance()
 	static bool bIsFirstInstance = false;
 	static bool bNeverFirst = FParse::Param(FCommandLine::Get(), TEXT("neverfirst"));
 
-#if !(UE_BUILD_SHIPPING && WITH_EDITOR)
 	if (!bIsFirstInstance && !bNeverFirst)	// once we determined that we're first, this can never change until we exit; otherwise, we re-check each time
 	{
 		// create the file if it doesn't exist
@@ -2036,13 +2035,12 @@ bool FUnixPlatformProcess::IsFirstInstance()
 			}
 		}
 	}
-#endif
+
 	return bIsFirstInstance;
 }
 
 void FUnixPlatformProcess::CeaseBeingFirstInstance()
 {
-#if !(UE_BUILD_SHIPPING && WITH_EDITOR)
 	if (GFileLockDescriptor != -1)
 	{
 		// may fail if we didn't have the lock
@@ -2050,7 +2048,6 @@ void FUnixPlatformProcess::CeaseBeingFirstInstance()
 		close(GFileLockDescriptor);
 		GFileLockDescriptor = -1;
 	}
-#endif
 }
 
 void FUnixPlatformProcess::OnChildEndFramePostFork()
diff --git a/Engine/Source/Runtime/Core/Private/Windows/WindowsPlatformMisc.cpp b/Engine/Source/Runtime/Core/Private/Windows/WindowsPlatformMisc.cpp
--- a/Engine/Source/Runtime/Core/Private/Windows/WindowsPlatformMisc.cpp
+++ b/Engine/Source/Runtime/Core/Private/Windows/WindowsPlatformMisc.cpp
@@ -574,6 +574,10 @@ void FWindowsPlatformMisc::PlatformInit()
 	FWindowsPlatformStackWalk::RegisterOnModulesChanged();
 }
 
+void FWindowsPlatformMisc::PlatformTearDown()
+{
+	FPlatformProcess::CeaseBeingFirstInstance();
+}
 
 /**
  * Handler called for console events like closure, CTRL-C, ...
diff --git a/Engine/Source/Runtime/Core/Private/Windows/WindowsPlatformProcess.cpp b/Engine/Source/Runtime/Core/Private/Windows/WindowsPlatformProcess.cpp
--- a/Engine/Source/Runtime/Core/Private/Windows/WindowsPlatformProcess.cpp
+++ b/Engine/Source/Runtime/Core/Private/Windows/WindowsPlatformProcess.cpp
@@ -45,6 +45,49 @@ TArray<FString> FWindowsPlatformProcess::DllDirectories;
 bool IsJobObjectSet = false;
 HANDLE GhJob = NULL;
 
+namespace WindowsPlatformProcess
+{
+	/**
+	 * Maintain a named mutex to detect whether we are the first instance of this game
+	 */
+	static HANDLE GNamedMutex = NULL;
+
+	void ReleaseNamedMutex(void)
+	{
+		if (GNamedMutex)
+		{
+			ReleaseMutex(GNamedMutex);
+			GNamedMutex = NULL;
+		}
+	}
+
+	bool MakeNamedMutex(const TCHAR* CmdLine)
+	{
+		bool bIsFirstInstance = false;
+
+		TCHAR MutexName[MAX_SPRINTF] = TEXT("");
+
+		FCString::Strcpy(MutexName, MAX_SPRINTF, TEXT("UnrealEngine4"));
+
+		GNamedMutex = CreateMutex(NULL, true, MutexName);
+
+		if (GNamedMutex && GetLastError() != ERROR_ALREADY_EXISTS && !FParse::Param(CmdLine, TEXT("NEVERFIRST")))
+		{
+			// We're the first instance!
+			bIsFirstInstance = true;
+		}
+		else
+		{
+			// Still need to release it in this case, because it gave us a valid copy
+			ReleaseNamedMutex();
+			// There is already another instance of the game running.
+			bIsFirstInstance = false;
+		}
+
+		return(bIsFirstInstance);
+	}
+}
+
 void FWindowsPlatformProcess::AddDllDirectory(const TCHAR* Directory)
 {
 	FString NormalizedDirectory = FPaths::ConvertRelativePathToFull(Directory);
@@ -2057,6 +2100,20 @@ void FWindowsPlatformProcess::SetupGameThread()
 	SetThreadName(TEXT("GameThread"));
 }
 
+bool FWindowsPlatformProcess::IsFirstInstance()
+{
+	// Named mutex we use to figure out whether we are the first instance of the game running. This is needed to e.g.
+	// make sure there is no contention when trying to save the shader cache.
+	static bool bIsFirstInstance = WindowsPlatformProcess::MakeNamedMutex(FCommandLine::Get());
+	return bIsFirstInstance;
+}
+
+void FWindowsPlatformProcess::CeaseBeingFirstInstance()
+{
+	// Release the mutex in the error case to ensure subsequent runs don't find it.
+	WindowsPlatformProcess::ReleaseNamedMutex();
+}
+
 namespace WindowsPlatformProcessImpl
 {
 	static void SetThreadName(LPCSTR ThreadName)
diff --git a/Engine/Source/Runtime/Core/Public/CoreGlobals.h b/Engine/Source/Runtime/Core/Public/CoreGlobals.h
--- a/Engine/Source/Runtime/Core/Public/CoreGlobals.h
+++ b/Engine/Source/Runtime/Core/Public/CoreGlobals.h
@@ -465,17 +465,6 @@ extern CORE_API uint32 GFrameNumber;
 /** NEED TO RENAME, for RT version of GFrameTime use View.ViewFamily->FrameNumber or pass down from RT from GFrameTime). */
 extern CORE_API uint32 GFrameNumberRenderThread;
 
-#if !(UE_BUILD_SHIPPING && WITH_EDITOR)
-// We cannot count on this variable to be accurate in a shipped game, so make sure no code tries to use it
-/** Whether we are the first instance of the game running. */
-#if PLATFORM_UNIX
-#define GIsFirstInstance FPlatformProcess::IsFirstInstance()
-#else
-extern CORE_API bool GIsFirstInstance;
-#endif
-
-#endif
-
 /** Threshold for a frame to be considered a hitch (in milliseconds). */
 extern CORE_API float GHitchThresholdMS;
 
diff --git a/Engine/Source/Runtime/Core/Public/Windows/WindowsPlatformMisc.h b/Engine/Source/Runtime/Core/Public/Windows/WindowsPlatformMisc.h
--- a/Engine/Source/Runtime/Core/Public/Windows/WindowsPlatformMisc.h
+++ b/Engine/Source/Runtime/Core/Public/Windows/WindowsPlatformMisc.h
@@ -46,6 +46,7 @@ struct CORE_API FWindowsPlatformMisc
 {
 	static void PlatformPreInit();
 	static void PlatformInit();
+	static void PlatformTearDown();
 	static void SetGracefulTerminationHandler();
 	static void CallGracefulTerminationHandler();
 	static ECrashHandlingType GetCrashHandlingType();
diff --git a/Engine/Source/Runtime/Core/Public/Windows/WindowsPlatformProcess.h b/Engine/Source/Runtime/Core/Public/Windows/WindowsPlatformProcess.h
--- a/Engine/Source/Runtime/Core/Public/Windows/WindowsPlatformProcess.h
+++ b/Engine/Source/Runtime/Core/Public/Windows/WindowsPlatformProcess.h
@@ -197,6 +197,13 @@ public:
 	static void SetupGameThread();
 	static void SetupAudioThread();
 	static void TeardownAudioThread();
+	static bool IsFirstInstance();
+
+	/**
+	 * @brief Releases locks that we held for IsFirstInstance check
+	 */
+	static void CeaseBeingFirstInstance();
+
 protected:
 
 	/**
diff --git a/Engine/Source/Runtime/Launch/Private/LaunchEngineLoop.cpp b/Engine/Source/Runtime/Launch/Private/LaunchEngineLoop.cpp
--- a/Engine/Source/Runtime/Launch/Private/LaunchEngineLoop.cpp
+++ b/Engine/Source/Runtime/Launch/Private/LaunchEngineLoop.cpp
@@ -2727,11 +2727,10 @@ PRAGMA_ENABLE_DEPRECATION_WARNINGS
 
 	{
 		bool bShouldCleanShaderWorkingDirectory = true;
-#if !(UE_BUILD_SHIPPING && WITH_EDITOR)
+
 		// Only clean the shader working directory if we are the first instance, to avoid deleting files in use by other instances
 		//@todo - check if any other instances are running right now
-		bShouldCleanShaderWorkingDirectory = GIsFirstInstance;
-#endif
+		bShouldCleanShaderWorkingDirectory = FPlatformProcess::IsFirstInstance();
 
 		if (bShouldCleanShaderWorkingDirectory && !FParse::Param(FCommandLine::Get(), TEXT("Multiprocess")))
 		{
diff --git a/Engine/Source/Runtime/Launch/Private/Windows/LaunchWindows.cpp b/Engine/Source/Runtime/Launch/Private/Windows/LaunchWindows.cpp
--- a/Engine/Source/Runtime/Launch/Private/Windows/LaunchWindows.cpp
+++ b/Engine/Source/Runtime/Launch/Private/Windows/LaunchWindows.cpp
@@ -40,49 +40,9 @@ extern "C" { _declspec(dllexport) extern const UINT D3D12SDKVersion = 602; }
 extern "C" { _declspec(dllexport) extern const char* D3D12SDKPath = ".\\D3D12\\"; }
 #endif // USE_D3D12_REDIST
 
-/**
- * Maintain a named mutex to detect whether we are the first instance of this game
- */
-HANDLE GNamedMutex = NULL;
-
 /** Whether we should pause before exiting. used by UCC */
 bool		GShouldPauseBeforeExit;
 
-void ReleaseNamedMutex( void )
-{
-	if( GNamedMutex )
-	{
-		ReleaseMutex( GNamedMutex );
-		GNamedMutex = NULL;
-	}
-}
-
-bool MakeNamedMutex( const TCHAR* CmdLine )
-{
-	bool bIsFirstInstance = false;
-
-	TCHAR MutexName[MAX_SPRINTF] = TEXT( "" );
-
-	FCString::Strcpy( MutexName, MAX_SPRINTF, TEXT( "UnrealEngine4" ) );
-
-	GNamedMutex = CreateMutex( NULL, true, MutexName );
-
-	if( GNamedMutex	&& GetLastError() != ERROR_ALREADY_EXISTS && !FParse::Param( CmdLine, TEXT( "NEVERFIRST" ) ) )
-	{
-		// We're the first instance!
-		bIsFirstInstance = true;
-	}
-	else
-	{
-		// Still need to release it in this case, because it gave us a valid copy
-		ReleaseNamedMutex();
-		// There is already another instance of the game running.
-		bIsFirstInstance = false;
-	}
-
-	return( bIsFirstInstance );
-}
-
 /**
  * Handler for CRT parameter validation. Triggers error
  *
@@ -234,16 +194,10 @@ LAUNCH_API int32 LaunchWindowsStartup( HINSTANCE hInInstance, HINSTANCE hPrevIns
 		SetErrorMode(SEM_FAILCRITICALERRORS | SEM_NOGPFAULTERRORBOX | SEM_NOOPENFILEERRORBOX);
 	}
 
-#if !(UE_BUILD_SHIPPING && WITH_EDITOR)
-	// Named mutex we use to figure out whether we are the first instance of the game running. This is needed to e.g.
-	// make sure there is no contention when trying to save the shader cache.
-	GIsFirstInstance = MakeNamedMutex( CmdLine );
-
 	if ( FParse::Param( CmdLine,TEXT("crashreports") ) )
 	{
 		GAlwaysReportCrash = true;
 	}
-#endif
 
 	bool bNoExceptionHandler = FParse::Param(CmdLine,TEXT("noexceptionhandler"));
 	(void)bNoExceptionHandler;
@@ -293,10 +247,6 @@ LAUNCH_API int32 LaunchWindowsStartup( HINSTANCE hInInstance, HINSTANCE hPrevIns
 #if !PLATFORM_SEH_EXCEPTIONS_DISABLED
 		__except( GEnableInnerException ? EXCEPTION_EXECUTE_HANDLER : ReportCrash( GetExceptionInformation( ) ) )
 		{
-#if !(UE_BUILD_SHIPPING && WITH_EDITOR)
-			// Release the mutex in the error case to ensure subsequent runs don't find it.
-			ReleaseNamedMutex();
-#endif
 			// Crashed.
 			ErrorLevel = 1;
 			if(GError)
@@ -320,11 +270,6 @@ LAUNCH_API void LaunchWindowsShutdown()
 	// Final shut down.
 	FEngineLoop::AppExit();
 
-#if !(UE_BUILD_SHIPPING && WITH_EDITOR)
-	// Release the named mutex again now that we are done.
-	ReleaseNamedMutex();
-#endif
-
 	// pause if we should
 	if (GShouldPauseBeforeExit)
 	{
