Fixed GIsConsoleExecutable by using ImageNtHeader

When stdout was redirected to a file or a pipe, GIsConsoleExecutable was false. That caused GLogConsole to incorrectly remain registered, and negatively impacted performance with its calls to WriteConsole on the main thread.

diff --git a/Engine/Source/Runtime/Launch/Private/Launch.cpp b/Engine/Source/Runtime/Launch/Private/Launch.cpp
--- a/Engine/Source/Runtime/Launch/Private/Launch.cpp
+++ b/Engine/Source/Runtime/Launch/Private/Launch.cpp
@@ -18,6 +18,9 @@
 #endif
 #if PLATFORM_WINDOWS
 	#include "Windows/WindowsHWrapper.h"
+	#include "Windows/AllowWindowsPlatformTypes.h"
+	#include <DbgHelp.h>
+	#include "Windows/HideWindowsPlatformTypes.h"
 #endif
 
 
@@ -134,7 +137,14 @@ int32 GuardedMain( const TCHAR* CmdLine )
 #if PLATFORM_WINDOWS
 	FCString::Strcpy(MiniDumpFilenameW, *FString::Printf(TEXT("unreal-v%i-%s.dmp"), FEngineVersion::Current().GetChangelist(), *FDateTime::Now().ToString()));
 
-	GIsConsoleExecutable = (GetFileType(GetStdHandle(STD_OUTPUT_HANDLE)) == FILE_TYPE_CHAR);
+	if (PIMAGE_NT_HEADERS NtHeaders = ImageNtHeader(GetModuleHandle(nullptr)))
+	{
+		GIsConsoleExecutable = (NtHeaders->OptionalHeader.Subsystem == IMAGE_SUBSYSTEM_WINDOWS_CUI);
+	}
+	else
+	{
+		GIsConsoleExecutable = (GetFileType(GetStdHandle(STD_OUTPUT_HANDLE)) == FILE_TYPE_CHAR);
+	}
 #endif
 
 	int32 ErrorLevel = EnginePreInit( CmdLine );
