If the CSV Profiler is running, output the CSV GT/RT frame in the GT and RT frame named events rather than the engine frame numbers, and prefix with Csv.

diff --git a/Engine/Source/Runtime/Core/Private/ProfilingDebugging/CsvProfiler.cpp b/Engine/Source/Runtime/Core/Private/ProfilingDebugging/CsvProfiler.cpp
--- a/Engine/Source/Runtime/Core/Private/ProfilingDebugging/CsvProfiler.cpp
+++ b/Engine/Source/Runtime/Core/Private/ProfilingDebugging/CsvProfiler.cpp
@@ -2712,6 +2712,7 @@ FCsvProfiler* FCsvProfiler::Get()
 FCsvProfiler::FCsvProfiler()
 	: NumFramesToCapture(-1)
 	, CaptureFrameNumber(0)
+	, CaptureFrameNumberRT(0)
 	, CaptureOnEventFrameCount(-1)
 	, bInsertEndFrameAtFrameStart(false)
 	, LastEndFrameTimestamp(0)
@@ -2868,6 +2869,7 @@ void FCsvProfiler::BeginFrame()
 					NumFramesToCapture = CurrentCommand.Value;
 					GCsvRepeatFrameCount = NumFramesToCapture;
 					CaptureFrameNumber = 0;
+					CaptureFrameNumberRT = 0;
 					LastEndFrameTimestamp = FPlatformTime::Cycles64();
 					CurrentFlags = CurrentCommand.Flags;
 
@@ -3142,6 +3144,10 @@ void FCsvProfiler::EndFrameRT()
 {
 	LLM_SCOPE(ELLMTag::CsvProfiler);
 	check(IsInRenderingThread());
+	if (GCsvProfilerIsCapturing)
+	{
+		CaptureFrameNumberRT++;
+	}
 }
 
 void FCsvProfiler::BeginCapture(int InNumFramesToCapture, 
@@ -3720,6 +3726,12 @@ int32 FCsvProfiler::GetCaptureFrameNumber()
 	return CaptureFrameNumber;
 }
 
+int32 FCsvProfiler::GetCaptureFrameNumberRT()
+{
+	return CaptureFrameNumberRT;
+}
+
+
 //Get the total frame to capture when we are capturing on event. 
 //Example:  -csvStartOnEvent="My Event"
 //			-csvCaptureOnEventFrameCount=2500
diff --git a/Engine/Source/Runtime/Core/Public/ProfilingDebugging/CsvProfiler.h b/Engine/Source/Runtime/Core/Public/ProfilingDebugging/CsvProfiler.h
--- a/Engine/Source/Runtime/Core/Public/ProfilingDebugging/CsvProfiler.h
+++ b/Engine/Source/Runtime/Core/Public/ProfilingDebugging/CsvProfiler.h
@@ -379,6 +379,7 @@ public:
 	CORE_API bool IsWritingFile();
 
 	CORE_API int32 GetCaptureFrameNumber();
+	CORE_API int32 GetCaptureFrameNumberRT();
 	CORE_API int32 GetNumFrameToCaptureOnEvent();
 
 	CORE_API bool EnableCategoryByString(const FString& CategoryName) const;
@@ -457,6 +458,7 @@ private:
 
 	int32 NumFramesToCapture;
 	int32 CaptureFrameNumber;
+	int32 CaptureFrameNumberRT;
 	int32 CaptureOnEventFrameCount;
 
 	bool bInsertEndFrameAtFrameStart;
diff --git a/Engine/Source/Runtime/Launch/Private/LaunchEngineLoop.cpp b/Engine/Source/Runtime/Launch/Private/LaunchEngineLoop.cpp
--- a/Engine/Source/Runtime/Launch/Private/LaunchEngineLoop.cpp
+++ b/Engine/Source/Runtime/Launch/Private/LaunchEngineLoop.cpp
@@ -4978,7 +4978,17 @@ static inline void BeginFrameRenderThread(FRHICommandListImmediate& RHICmdList,
 	}
 #endif //CPUPROFILERTRACE_ENABLED
 
-	FString FrameString = FString::Printf(TEXT("Frame %d"), CurrentFrameCounter);
+	FString FrameString;
+#if CSV_PROFILER
+	if (FCsvProfiler::Get()->IsCapturing())
+	{
+		FrameString = FString::Printf(TEXT("CsvFrame %d"), FCsvProfiler::Get()->GetCaptureFrameNumberRT());
+	}
+	else
+#endif
+	{
+		FrameString = FString::Printf(TEXT("Frame %d"), CurrentFrameCounter);
+	}
 #if ENABLE_NAMED_EVENTS
 #if PLATFORM_LIMIT_PROFILER_UNIQUE_NAMED_EVENTS
 	FPlatformMisc::BeginNamedEvent(FColor::Yellow, TEXT("Frame"));
@@ -5119,10 +5129,19 @@ void FEngineLoop::Tick()
 	{
 #if PLATFORM_LIMIT_PROFILER_UNIQUE_NAMED_EVENTS
 		FrameString = TEXT("FEngineLoop");
-#else
-		FCString::Snprintf(IndexedFrameString, 32, TEXT("Frame %d"), CurrentFrameCounter);
-		FrameString = IndexedFrameString;
+#else // PLATFORM_LIMIT_PROFILER_UNIQUE_NAMED_EVENTS
+#if CSV_PROFILER
+		if (FCsvProfiler::Get()->IsCapturing())
+		{
+			FCString::Snprintf(IndexedFrameString, 32, TEXT("CsvFrame %d"), FCsvProfiler::Get()->GetCaptureFrameNumber());
+		}
+		else
 #endif
+		{
+			FCString::Snprintf(IndexedFrameString, 32, TEXT("Frame %d"), CurrentFrameCounter);
+		}
+		FrameString = IndexedFrameString;
+#endif // !PLATFORM_LIMIT_PROFILER_UNIQUE_NAMED_EVENTS
 	}
 	SCOPED_NAMED_EVENT_TCHAR(FrameString, FColor::Red);
 #endif
@@ -6428,4 +6447,4 @@ void FPreInitContext::Cleanup()
 	SlowTaskPtr = nullptr;
 }
 
-#undef LOCTEXT_NAMESPACE
\ No newline at end of file
+#undef LOCTEXT_NAMESPACE
