Partial backout to workaround assert

diff --git a/Engine/Source/Runtime/Launch/Private/LaunchEngineLoop.cpp b/Engine/Source/Runtime/Launch/Private/LaunchEngineLoop.cpp
--- a/Engine/Source/Runtime/Launch/Private/LaunchEngineLoop.cpp
+++ b/Engine/Source/Runtime/Launch/Private/LaunchEngineLoop.cpp
@@ -4978,17 +4978,7 @@ static inline void BeginFrameRenderThread(FRHICommandListImmediate& RHICmdList,
 	}
 #endif //CPUPROFILERTRACE_ENABLED
 
-	FString FrameString;
-#if CSV_PROFILER
-	if (FCsvProfiler::Get()->IsCapturing())
-	{
-		FrameString = FString::Printf(TEXT("CsvFrame %d"), FCsvProfiler::Get()->GetCaptureFrameNumberRT());
-	}
-	else
-#endif
-	{
-		FrameString = FString::Printf(TEXT("Frame %d"), CurrentFrameCounter);
-	}
+	FString FrameString = FString::Printf(TEXT("Frame %d"), CurrentFrameCounter);
 #if ENABLE_NAMED_EVENTS
 #if PLATFORM_LIMIT_PROFILER_UNIQUE_NAMED_EVENTS
 	FPlatformMisc::BeginNamedEvent(FColor::Yellow, TEXT("Frame"));
@@ -5129,19 +5119,10 @@ void FEngineLoop::Tick()
 	{
 #if PLATFORM_LIMIT_PROFILER_UNIQUE_NAMED_EVENTS
 		FrameString = TEXT("FEngineLoop");
-#else // PLATFORM_LIMIT_PROFILER_UNIQUE_NAMED_EVENTS
-#if CSV_PROFILER
-		if (FCsvProfiler::Get()->IsCapturing())
-		{
-			FCString::Snprintf(IndexedFrameString, 32, TEXT("CsvFrame %d"), FCsvProfiler::Get()->GetCaptureFrameNumber());
-		}
-		else
-#endif
-		{
-			FCString::Snprintf(IndexedFrameString, 32, TEXT("Frame %d"), CurrentFrameCounter);
-		}
+#else
+		FCString::Snprintf(IndexedFrameString, 32, TEXT("Frame %d"), CurrentFrameCounter);
 		FrameString = IndexedFrameString;
-#endif // !PLATFORM_LIMIT_PROFILER_UNIQUE_NAMED_EVENTS
+#endif
 	}
 	SCOPED_NAMED_EVENT_TCHAR(FrameString, FColor::Red);
 #endif
