Reinstate 21355934 with assert fix

diff --git a/Engine/Source/Runtime/Launch/Private/LaunchEngineLoop.cpp b/Engine/Source/Runtime/Launch/Private/LaunchEngineLoop.cpp
--- a/Engine/Source/Runtime/Launch/Private/LaunchEngineLoop.cpp
+++ b/Engine/Source/Runtime/Launch/Private/LaunchEngineLoop.cpp
@@ -4978,7 +4978,17 @@ static inline void BeginFrameRenderThread(FRHICommandListImmediate& RHICmdList,
 	}
 #endif //CPUPROFILERTRACE_ENABLED
 
-	FString FrameString = FString::Printf(TEXT("Frame %d"), CurrentFrameCounter);
+	FString FrameString;
+#if CSV_PROFILER
+	if (FCsvProfiler::Get()->IsCapturing_Renderthread())
+	{
+		FrameString = FString::Printf(TEXT("CsvFrame %d"), FCsvProfiler::Get()->GetCaptureFrameNumberRT());
+	}
+	else
+#endif
+	{
+		FrameString = FString::Printf(TEXT("Frame %d"), CurrentFrameCounter);
+	}
 #if ENABLE_NAMED_EVENTS
 #if PLATFORM_LIMIT_PROFILER_UNIQUE_NAMED_EVENTS
 	FPlatformMisc::BeginNamedEvent(FColor::Yellow, TEXT("Frame"));
@@ -5119,10 +5129,19 @@ void FEngineLoop::Tick()
 	{
 #if PLATFORM_LIMIT_PROFILER_UNIQUE_NAMED_EVENTS
 		FrameString = TEXT("FEngineLoop");
-#else
-		FCString::Snprintf(IndexedFrameString, 32, TEXT("Frame %d"), CurrentFrameCounter);
-		FrameString = IndexedFrameString;
+#else // PLATFORM_LIMIT_PROFILER_UNIQUE_NAMED_EVENTS
+#if CSV_PROFILER
+		if (FCsvProfiler::Get()->IsCapturing())
+		{
+			FCString::Snprintf(IndexedFrameString, 32, TEXT("CsvFrame %d"), FCsvProfiler::Get()->GetCaptureFrameNumber());
+		}
+		else
 #endif
+		{
+			FCString::Snprintf(IndexedFrameString, 32, TEXT("Frame %d"), CurrentFrameCounter);
+		}
+		FrameString = IndexedFrameString;
+#endif // !PLATFORM_LIMIT_PROFILER_UNIQUE_NAMED_EVENTS
 	}
 	SCOPED_NAMED_EVENT_TCHAR(FrameString, FColor::Red);
 #endif
