Added UE_SCOPED_ENGINE_ACTIVITY in strategical places to cover most of the cost when starting the various runtimes

diff --git a/Engine/Plugins/Experimental/GameFeatures/Source/GameFeatures/Private/GameFeaturesSubsystem.cpp b/Engine/Plugins/Experimental/GameFeatures/Source/GameFeatures/Private/GameFeaturesSubsystem.cpp
--- a/Engine/Plugins/Experimental/GameFeatures/Source/GameFeatures/Private/GameFeaturesSubsystem.cpp
+++ b/Engine/Plugins/Experimental/GameFeatures/Source/GameFeatures/Private/GameFeaturesSubsystem.cpp
@@ -1041,6 +1041,8 @@ void UGameFeaturesSubsystem::CancelGameFeatureStateChange(const FString& PluginU
 
 void UGameFeaturesSubsystem::LoadBuiltInGameFeaturePlugin(const TSharedRef<IPlugin>& Plugin, FBuiltInPluginAdditionalFilters AdditionalFilter)
 {
+	UE_SCOPED_ENGINE_ACTIVITY(TEXT("Loading GameFeaturePlugin %s"), *Plugin->GetName());
+
 	UAssetManager::Get().PushBulkScanning();
 
 	const FString& PluginDescriptorFilename = Plugin->GetDescriptorFileName();
diff --git a/Engine/Source/Runtime/Core/Private/Internationalization/TextLocalizationManager.cpp b/Engine/Source/Runtime/Core/Private/Internationalization/TextLocalizationManager.cpp
--- a/Engine/Source/Runtime/Core/Private/Internationalization/TextLocalizationManager.cpp
+++ b/Engine/Source/Runtime/Core/Private/Internationalization/TextLocalizationManager.cpp
@@ -341,7 +341,7 @@ void BeginInitTextLocalization()
 void InitEngineTextLocalization()
 {
 	LLM_SCOPE(ELLMTag::Localization);
-	
+	UE_SCOPED_ENGINE_ACTIVITY(TEXT("Initializing Localization"));
 	SCOPED_BOOT_TIMING("InitEngineTextLocalization");
 	DECLARE_SCOPE_CYCLE_COUNTER(TEXT("InitEngineTextLocalization"), STAT_InitEngineTextLocalization, STATGROUP_LoadTime);
 
diff --git a/Engine/Source/Runtime/Core/Private/Modules/ModuleManager.cpp b/Engine/Source/Runtime/Core/Private/Modules/ModuleManager.cpp
--- a/Engine/Source/Runtime/Core/Private/Modules/ModuleManager.cpp
+++ b/Engine/Source/Runtime/Core/Private/Modules/ModuleManager.cpp
@@ -428,6 +428,7 @@ IModuleInterface* FModuleManager::LoadModuleWithFailureReason(const FName InModu
 	// enable this ensure when we know it's okay
 //	ensureMsgf(IsInGameThread(), TEXT("ModuleManager: Attempting to load '%s' outside the main thread.  Please call LoadModule on the main/game thread only.  You can use GetModule or GetModuleChecked instead, those are safe to call outside the game thread."), *InModuleName.ToString());
 
+	UE_SCOPED_ENGINE_ACTIVITY(TEXT("Loading Module %s"), *InModuleName.ToString());
 	SCOPED_BOOT_TIMING("LoadModule");
 	TRACE_CPUPROFILER_EVENT_SCOPE_TEXT(*FString::Printf(TEXT("LoadModule_%s"), *InModuleName.ToString()));
 	DECLARE_SCOPE_CYCLE_COUNTER(TEXT("Module Load"), STAT_ModuleLoad, STATGROUP_LoadTime);
diff --git a/Engine/Source/Runtime/Engine/Private/Subsystems/SubsystemCollection.cpp b/Engine/Source/Runtime/Engine/Private/Subsystems/SubsystemCollection.cpp
--- a/Engine/Source/Runtime/Engine/Private/Subsystems/SubsystemCollection.cpp
+++ b/Engine/Source/Runtime/Engine/Private/Subsystems/SubsystemCollection.cpp
@@ -224,6 +224,8 @@ void FSubsystemCollectionBase::AddReferencedObjects(UObject* Referencer, FRefere
 
 USubsystem* FSubsystemCollectionBase::AddAndInitializeSubsystem(UClass* SubsystemClass)
 {
+	UE_SCOPED_ENGINE_ACTIVITY(TEXT("Initializing Subsystem %s"), *SubsystemClass->GetName());
+
 	TGuardValue<bool> PopulatingGuard(bPopulating, true);
 
 	if (!SubsystemMap.Contains(SubsystemClass))
diff --git a/Engine/Source/Runtime/Engine/Private/UnrealEngine.cpp b/Engine/Source/Runtime/Engine/Private/UnrealEngine.cpp
--- a/Engine/Source/Runtime/Engine/Private/UnrealEngine.cpp
+++ b/Engine/Source/Runtime/Engine/Private/UnrealEngine.cpp
@@ -2962,6 +2962,7 @@ static FAutoConsoleCommand GResetTimecode(
 */
 void UEngine::InitializeObjectReferences()
 {
+	UE_SCOPED_ENGINE_ACTIVITY(TEXT("Initializing Object References"));
 	TRACE_LOADTIME_REQUEST_GROUP_SCOPE(TEXT("UEngine::InitializeObjectReferences"));
 	SCOPED_BOOT_TIMING("UEngine::InitializeObjectReferences");
 	DECLARE_SCOPE_CYCLE_COUNTER(TEXT("UEngine::InitializeObjectReferences"), STAT_InitializeObjectReferences, STATGROUP_LoadTime);
@@ -14558,6 +14559,7 @@ bool UEngine::LoadMap( FWorldContext& WorldContext, FURL URL, class UPendingNetG
 		// If the level isn't already in memory, load level from disk
 		if (WorldPackage == nullptr)
 		{
+			UE_SCOPED_ENGINE_ACTIVITY(TEXT("Loading World %s"), *URL.Map);
 			WorldPackage = LoadPackage(nullptr, *URL.Map, (WorldContext.WorldType == EWorldType::PIE ? LOAD_PackageForPIE : LOAD_None));
 		}
 
diff --git a/Engine/Source/Runtime/Launch/Private/Launch.cpp b/Engine/Source/Runtime/Launch/Private/Launch.cpp
--- a/Engine/Source/Runtime/Launch/Private/Launch.cpp
+++ b/Engine/Source/Runtime/Launch/Private/Launch.cpp
@@ -13,6 +13,7 @@
 #include "Misc/CoreDelegates.h"
 #include "Misc/EngineVersion.h"
 #include "Misc/ScopedSlowTask.h"
+#include "Misc/TrackedActivity.h"
 #if WITH_EDITOR
 	#include "UnrealEdGlobals.h"
 #endif
@@ -95,6 +96,8 @@ extern UNREALED_API FSecondsCounterData BlueprintCompileAndLoadTimerData;
  */
 int32 GuardedMain( const TCHAR* CmdLine )
 {
+	FTrackedActivity::GetEngineActivity().Update(TEXT("Starting"), FTrackedActivity::ELight::Yellow);
+
 	FTaskTagScope Scope(ETaskTag::EGameThread);
 
 #if !(UE_BUILD_SHIPPING)
@@ -147,6 +150,7 @@ int32 GuardedMain( const TCHAR* CmdLine )
 	}
 #endif
 
+	FTrackedActivity::GetEngineActivity().Update(TEXT("Initializing"));
 	int32 ErrorLevel = EnginePreInit( CmdLine );
 
 	// exit if PreInit failed.
@@ -189,6 +193,8 @@ int32 GuardedMain( const TCHAR* CmdLine )
 	BootTimingPoint("Tick loop starting");
 	DumpBootTiming();
 
+	FTrackedActivity::GetEngineActivity().Update(TEXT("Ticking loop"), FTrackedActivity::ELight::Green);
+
 	// Don't tick if we're running an embedded engine - we rely on the outer
 	// application ticking us instead.
 	if (!GUELibraryOverrideSettings.bIsEmbedded)
