[engine] Fixed saved directory paths being cached before -saveddir etc are processed
#rnx

diff --git a/Engine/Source/Runtime/Launch/Private/LaunchEngineLoop.cpp b/Engine/Source/Runtime/Launch/Private/LaunchEngineLoop.cpp
--- a/Engine/Source/Runtime/Launch/Private/LaunchEngineLoop.cpp
+++ b/Engine/Source/Runtime/Launch/Private/LaunchEngineLoop.cpp
@@ -1811,16 +1811,14 @@ int32 FEngineLoop::PreInitPreStartupScreen(const TCHAR* CmdLine)
 		FConfigFile CommandLineAliasesConfigFile;
 		if (FPaths::ProjectDir().Len() > 0)
 		{
-			// Pass GeneratedConfigDir as nullptr instead of FPaths::GeneratedConfigDir() so that saved directory is not cached before -saveddir argument can be added
-			const TCHAR* GeneratedConfigDir = nullptr;
-
+			// Pass EngineIntermediateDir as GeneratedConfigDir of FPaths::GeneratedConfigDir() so that saved directory is not cached before -saveddir argument can be added
 			FConfigCacheIni::LoadExternalIniFile(CommandLineAliasesConfigFile, TEXT("CommandLineAliases"), nullptr, *FPaths::Combine(FPaths::ProjectDir(), TEXT("Config")),
 												 /*bIsBaseIniName=*/ false,
 												 /*Platform=*/ nullptr,
 												 /*bForceReload=*/ false,
 												 /*bWriteDestIni=*/ false,
 												 /*bAllowGeneratedIniWhenCooked=*/ true,
-												 GeneratedConfigDir);
+												 /*GeneratedConfigDir=*/ *FPaths::EngineIntermediateDir());
 		}
 
 		TArray<FString> PrevAliasExpansions;
