[Runtime.Launch]
* LaunchEngineLoop - Added activity scopes and moved console to show earlier

diff --git a/Engine/Source/Runtime/Launch/Private/LaunchEngineLoop.cpp b/Engine/Source/Runtime/Launch/Private/LaunchEngineLoop.cpp
--- a/Engine/Source/Runtime/Launch/Private/LaunchEngineLoop.cpp
+++ b/Engine/Source/Runtime/Launch/Private/LaunchEngineLoop.cpp
@@ -31,6 +31,7 @@
 #include "Misc/App.h"
 #include "Misc/OutputDeviceConsole.h"
 #include "Misc/ScopeExit.h"
+#include "Misc/TrackedActivity.h"
 #include "HAL/PlatformFileManager.h"
 #include "HAL/FileManagerGeneric.h"
 #include "HAL/ExceptionHandling.h"
@@ -2835,6 +2836,7 @@ PRAGMA_ENABLE_DEPRECATION_WARNINGS
 	}
 
 	{
+		UE_SCOPED_ENGINE_ACTIVITY(TEXT("Initializing Render Settings"));
 		SCOPED_BOOT_TIMING("RenderUtilsInit");
 		// One-time initialization of global variables based on engine configuration.
 		RenderUtilsInit();
@@ -2881,6 +2883,7 @@ PRAGMA_ENABLE_DEPRECATION_WARNINGS
 #if WITH_EDITORONLY_DATA
 		{
 			// Ensure that DDC is initialized from the game thread.
+			UE_SCOPED_ENGINE_ACTIVITY(TEXT("Initializing Derived Data Cache"));
 			SCOPED_BOOT_TIMING("InitDerivedData");
 			UE::DerivedData::GetBuild();
 			UE::DerivedData::GetCache();
@@ -2927,6 +2930,7 @@ PRAGMA_ENABLE_DEPRECATION_WARNINGS
 	{
 		if (AllowShaderCompiling())
 		{
+			UE_SCOPED_ENGINE_ACTIVITY(TEXT("Initializing Shader Types"));
 			SCOPED_BOOT_TIMING("InitializeShaderTypes");
 			// Initialize shader types before loading any shaders
 			InitializeShaderTypes();
@@ -3212,6 +3216,7 @@ int32 FEngineLoop::PreInitPostStartupScreen(const TCHAR* CmdLine)
 			}
 			else
 			{
+				UE_SCOPED_ENGINE_ACTIVITY(TEXT("Play Preload Screen"));
 				SCOPED_BOOT_TIMING("PlayFirstPreLoadScreen");
 
 				if (FPreLoadScreenManager::Get())
@@ -3324,7 +3329,8 @@ int32 FEngineLoop::PreInitPostStartupScreen(const TCHAR* CmdLine)
 
 		{
 		    SCOPED_BOOT_TIMING("LoadAssetRegistryModule");
-		    // If we don't do this now and the async loading thread is active, then we will attempt to load this module from a thread
+			UE_SCOPED_ENGINE_ACTIVITY(TEXT("Loading AssetRegistry"));
+			// If we don't do this now and the async loading thread is active, then we will attempt to load this module from a thread
 		    FModuleManager::Get().LoadModule("AssetRegistry");
 		}
 #if WITH_COREUOBJECT
@@ -3344,7 +3350,10 @@ int32 FEngineLoop::PreInitPostStartupScreen(const TCHAR* CmdLine)
 		FDelayedAutoRegisterHelper::RunAndClearDelayedAutoRegisterDelegates(EDelayedRegisterRunPhase::PreObjectSystemReady);
 
 		// Make sure all UObject classes are registered and default properties have been initialized
-		ProcessNewlyLoadedUObjects();
+		{
+			UE_SCOPED_ENGINE_ACTIVITY(TEXT("Initializing UObject Classes"));
+			ProcessNewlyLoadedUObjects();
+		}
 
 		EndInitGameTextLocalization();
 
@@ -3723,7 +3732,11 @@ int32 FEngineLoop::PreInitPostStartupScreen(const TCHAR* CmdLine)
 #if WITH_ENGINE
 			PRIVATE_GRunningCommandletClass = CommandletClass;
 #endif
-			int32 ErrorLevel = Commandlet->Main( CommandletCommandLine );
+			int32 ErrorLevel;
+			{
+				FTrackedActivityScope CommandletActivity(FTrackedActivity::GetEngineActivity(), *FString::Printf(TEXT("Running %s"), *Commandlet->GetName()), false, FTrackedActivity::ELight::Green);
+				ErrorLevel = Commandlet->Main(CommandletCommandLine);
+			}
 #if WITH_ENGINE
 			PRIVATE_GRunningCommandletClass = nullptr;
 #endif
@@ -6118,6 +6131,8 @@ bool FEngineLoop::AppInit( )
 
 	{
 		SCOPED_BOOT_TIMING("Load pre-init plugin modules");
+		UE_SCOPED_ENGINE_ACTIVITY(TEXT("Loading Plugins (PreInit)"));
+
 		// Load "pre-init" plugin modules
 		if (!ProjectManager.LoadModulesForProject(ELoadingPhase::PostConfigInit) || !PluginManager.LoadModulesForEnabledPlugins(ELoadingPhase::PostConfigInit))
 		{
