Initialize device profile manager earlier in the pre init phase to properly set the FCoreDelegates::GatherDeviceProfileCVars.

diff --git a/Engine/Source/Runtime/Launch/Private/LaunchEngineLoop.cpp b/Engine/Source/Runtime/Launch/Private/LaunchEngineLoop.cpp
--- a/Engine/Source/Runtime/Launch/Private/LaunchEngineLoop.cpp
+++ b/Engine/Source/Runtime/Launch/Private/LaunchEngineLoop.cpp
@@ -2570,7 +2570,13 @@ PRAGMA_ENABLE_DEPRECATION_WARNINGS
 			// This may include scalability group settings which will override
 			// the defaults set above which can then be replaced below when
 			// the game user settings are loaded and applied.
+
+#if ALLOW_OTHER_PLATFORM_CONFIG
+			// Need to initialize the DeviceProfileManager Singleton to bind FCoreDelegates::GatherDeviceProfileCVars. This delegate is needed at cook time.
+			UDeviceProfileManager::Get().InitializeCVarsForActiveDeviceProfile();
+#else
 			UDeviceProfileManager::InitializeCVarsForActiveDeviceProfile();
+#endif
 		}
 
 		{
