If project path is set from uecommandline.txt the engine doesn't pick up the game name correctly.
Check first max 10 arguments instead of just the first.

diff --git a/Engine/Source/Runtime/Launch/Private/LaunchEngineLoop.cpp b/Engine/Source/Runtime/Launch/Private/LaunchEngineLoop.cpp
--- a/Engine/Source/Runtime/Launch/Private/LaunchEngineLoop.cpp
+++ b/Engine/Source/Runtime/Launch/Private/LaunchEngineLoop.cpp
@@ -516,36 +516,42 @@ void InitializeStdOutDevice()
 
 bool ParseGameProjectFromCommandLine(const TCHAR* InCmdLine, FString& OutProjectFilePath, FString& OutGameName)
 {
-	const TCHAR *CmdLine = InCmdLine;
-	FString FirstCommandLineToken = FParse::Token(CmdLine, 0);
-
-	// trim any whitespace at edges of string - this can happen if the token was quoted with leading or trailing whitespace
-	// VC++ tends to do this in its "external tools" config
-	FirstCommandLineToken.TrimStartInline();
-
-	//
 	OutProjectFilePath = TEXT("");
 	OutGameName = TEXT("");
 
-	if ( FirstCommandLineToken.Len() && !FirstCommandLineToken.StartsWith(TEXT("-")) )
+	const TCHAR *CmdLine = InCmdLine;
+	FString CommandLineToken = FParse::Token(CmdLine, 0);
+
+	// Avoid checking too many arguments
+	int MaxArgsChecked = 10;
+
+	while (!CommandLineToken.IsEmpty() && MaxArgsChecked-- > 0)
 	{
-		// The first command line argument could be the project file if it exists or the game name if not launching with a project file
-		const FString ProjectFilePath = FString(FirstCommandLineToken);
-		if ( FPaths::GetExtension(ProjectFilePath) == FProjectDescriptor::GetExtension() )
-		{
-			OutProjectFilePath = FirstCommandLineToken;
-			// Here we derive the game name from the project file
-			OutGameName = FPaths::GetBaseFilename(OutProjectFilePath);
-			return true;
-		}
-		else if (FPaths::IsRelative(FirstCommandLineToken) && FPlatformProperties::IsMonolithicBuild() == false)
+		// trim any whitespace at edges of string - this can happen if the token was quoted with leading or trailing whitespace
+		// VC++ tends to do this in its "external tools" config
+		CommandLineToken.TrimStartInline();
+
+		if (CommandLineToken.Len() && !CommandLineToken.StartsWith(TEXT("-")))
 		{
-			// Full game name is assumed to be the first token
-			OutGameName = MoveTemp(FirstCommandLineToken);
-			// Derive the project path from the game name. All games must have a uproject file, even if they are in the root folder.
-			OutProjectFilePath = FPaths::Combine(*FPaths::RootDir(), *OutGameName, *FString(OutGameName + TEXT(".") + FProjectDescriptor::GetExtension()));
-			return true;
+			// This command line argument could be the project file if it exists or the game name if not launching with a project file
+			const FString ProjectFilePath = FString(CommandLineToken);
+			if (FPaths::GetExtension(ProjectFilePath) == FProjectDescriptor::GetExtension())
+			{
+				OutProjectFilePath = CommandLineToken;
+				// Here we derive the game name from the project file
+				OutGameName = FPaths::GetBaseFilename(OutProjectFilePath);
+				return true;
+			}
+			else if (FPaths::IsRelative(CommandLineToken) && FPlatformProperties::IsMonolithicBuild() == false)
+			{
+				// Full game name is assumed to be the first token
+				OutGameName = MoveTemp(CommandLineToken);
+				// Derive the project path from the game name. All games must have a uproject file, even if they are in the root folder.
+				OutProjectFilePath = FPaths::Combine(*FPaths::RootDir(), *OutGameName, *FString(OutGameName + TEXT(".") + FProjectDescriptor::GetExtension()));
+				return true;
+			}
 		}
+		CommandLineToken = FParse::Token(CmdLine, 0);
 	}
 
 #if WITH_EDITOR
