Fix warning in cook by initializing the DeviceProfileManger. Calling the UDeviceProfileManager::Get().InitializeCVarsForActiveDeviceProfile() does intialize sampler states (Wrap_WorldGroupSettings and Clamp_WorldGroupSettings) that can be refresh by RefreshSamplerStatesCallback during the engine init phase. We need to check if the DynamicRHI was created before refreshing samplerstates.

diff --git a/Engine/Source/Runtime/Engine/Private/UnrealEngine.cpp b/Engine/Source/Runtime/Engine/Private/UnrealEngine.cpp
--- a/Engine/Source/Runtime/Engine/Private/UnrealEngine.cpp
+++ b/Engine/Source/Runtime/Engine/Private/UnrealEngine.cpp
@@ -998,7 +998,7 @@ void RefreshSamplerStatesCallback()
 		}
 
 		// The shared sampler states don't have an associated texture so must be manually refreshed
-		if (Wrap_WorldGroupSettings || Clamp_WorldGroupSettings)
+		if (GDynamicRHI && (Wrap_WorldGroupSettings || Clamp_WorldGroupSettings))
 		{
 			FSharedSamplerState* WrapState = Wrap_WorldGroupSettings;
 			FSharedSamplerState* ClampState = Clamp_WorldGroupSettings;
diff --git a/Engine/Source/Runtime/Launch/Private/LaunchEngineLoop.cpp b/Engine/Source/Runtime/Launch/Private/LaunchEngineLoop.cpp
--- a/Engine/Source/Runtime/Launch/Private/LaunchEngineLoop.cpp
+++ b/Engine/Source/Runtime/Launch/Private/LaunchEngineLoop.cpp
@@ -2576,7 +2576,13 @@ PRAGMA_ENABLE_DEPRECATION_WARNINGS
 			// This may include scalability group settings which will override
 			// the defaults set above which can then be replaced below when
 			// the game user settings are loaded and applied.
+
+#if ALLOW_OTHER_PLATFORM_CONFIG
+			// Need to initialize the DeviceProfileManager Singleton to bind FCoreDelegates::GatherDeviceProfileCVars. This delegate is needed at cook time.
+			UDeviceProfileManager::Get().InitializeCVarsForActiveDeviceProfile();
+#else
 			UDeviceProfileManager::InitializeCVarsForActiveDeviceProfile();
+#endif
 		}
 
 		{
