Undoing the change to use the shortened #include path for this file because it is in a public header.

diff --git a/Engine/Source/Runtime/Launch/Public/RequiredProgramMainCPPInclude.h b/Engine/Source/Runtime/Launch/Public/RequiredProgramMainCPPInclude.h
--- a/Engine/Source/Runtime/Launch/Public/RequiredProgramMainCPPInclude.h
+++ b/Engine/Source/Runtime/Launch/Public/RequiredProgramMainCPPInclude.h
@@ -6,4 +6,4 @@
 #include "LaunchEngineLoop.h"
 
 // this is highly sketchy, but we need some stuff from launchengineloop.cpp
-#include "LaunchEngineLoop.cpp"
+#include "Runtime/Launch/Private/LaunchEngineLoop.cpp"
