Delay the platform ShaderAutoGen. Thoses are generated header files needed when we compile shaders. We need to delay the platform initialization from the current platform since it needs to happen after the DeviceProfileManager creation. Move  ShaderAutoGenInit code out from LaunchEngineLoop.cp since it only needed from ShaderCompile.

diff --git a/Engine/Source/Runtime/Engine/Private/ShaderCompiler/ShaderCompiler.cpp b/Engine/Source/Runtime/Engine/Private/ShaderCompiler/ShaderCompiler.cpp
--- a/Engine/Source/Runtime/Engine/Private/ShaderCompiler/ShaderCompiler.cpp
+++ b/Engine/Source/Runtime/Engine/Private/ShaderCompiler/ShaderCompiler.cpp
@@ -70,6 +70,7 @@
 #include "UnrealEngine.h"
 
 #if WITH_EDITOR
+#include "Misc/DelayedAutoRegister.h"
 #include "Compression/OodleDataCompression.h"
 #include "UObject/ArchiveCookContext.h"
 #include "DerivedDataCache.h"
@@ -204,6 +205,42 @@ static TAutoConsoleVariable<bool> CVarShaderCompilerDumpDDCKeys(
 	ECVF_Default
 );
 
+#if WITH_EDITOR
+static FDelayedAutoRegisterHelper GKickOffShaderAutoGen(EDelayedRegisterRunPhase::ShaderTypesReady, []
+{
+	FShaderCompileUtilities::GenerateBrdfHeaders(GMaxRHIShaderPlatform);
+});
+
+static FDelayedAutoRegisterHelper GKickOffShaderAutoGenForPlatforms(EDelayedRegisterRunPhase::DeviceProfileManagerReady, []
+{
+	// also do this for all active target platforms (e.g. when cooking)
+	ITargetPlatformManagerModule* TPM = GetTargetPlatformManager();
+	if (TPM)
+	{
+		const TArray<ITargetPlatform*>& Platforms = TPM->GetActiveTargetPlatforms();
+
+		for (int32 Index = 0; Index < Platforms.Num(); ++Index)
+		{
+			TArray<FName> DesiredShaderFormats;
+			checkf(Platforms[Index], TEXT("Null platform on the list of active platforms!"));
+			Platforms[Index]->GetAllTargetedShaderFormats(DesiredShaderFormats);
+
+			for (int32 FormatIndex = 0; FormatIndex < DesiredShaderFormats.Num(); ++FormatIndex)
+			{
+				FShaderCompileUtilities::GenerateBrdfHeaders(DesiredShaderFormats[FormatIndex]);
+			}
+		}
+	}
+
+	// also do this for the editor mobile preview
+	EShaderPlatform MobilePreviewShaderPlatform = GShaderPlatformForFeatureLevel[ERHIFeatureLevel::ES3_1];
+	if (MobilePreviewShaderPlatform != SP_NumPlatforms)
+	{
+		FShaderCompileUtilities::GenerateBrdfHeaders(MobilePreviewShaderPlatform);
+	}
+});
+#endif
+
 /** Helper functions for logging more debug info */
 namespace ShaderCompiler
 {
diff --git a/Engine/Source/Runtime/Launch/Private/LaunchEngineLoop.cpp b/Engine/Source/Runtime/Launch/Private/LaunchEngineLoop.cpp
--- a/Engine/Source/Runtime/Launch/Private/LaunchEngineLoop.cpp
+++ b/Engine/Source/Runtime/Launch/Private/LaunchEngineLoop.cpp
@@ -1509,39 +1509,7 @@ static void UpdateGInputTime()
 	GInputTime = FPlatformTime::Cycles64();
 }
 
-#if WITH_EDITOR
-static void ShaderAutogenInit()
-{
-	// Force generation of the autogen files for the host platform
-	FShaderCompileUtilities::GenerateBrdfHeaders(GMaxRHIShaderPlatform);
-
-	// also do this for all active target platforms (e.g. when cooking)
-	ITargetPlatformManagerModule* TPM = GetTargetPlatformManager();
-	if (TPM)
-	{
-		const TArray<ITargetPlatform*>& Platforms = TPM->GetActiveTargetPlatforms();
-
-		for (int32 Index = 0; Index < Platforms.Num(); ++Index)
-		{
-			TArray<FName> DesiredShaderFormats;
-			checkf(Platforms[Index], TEXT("Null platform on the list of active platforms!"));
-			Platforms[Index]->GetAllTargetedShaderFormats(DesiredShaderFormats);
-
-			for (int32 FormatIndex = 0; FormatIndex < DesiredShaderFormats.Num(); ++FormatIndex)
-			{
-				FShaderCompileUtilities::GenerateBrdfHeaders(DesiredShaderFormats[FormatIndex]);
-			}
-		}
-	}
 
-	// also do this for the editor mobile preview
-	EShaderPlatform MobilePreviewShaderPlatform = GShaderPlatformForFeatureLevel[ERHIFeatureLevel::ES3_1];
-	if (MobilePreviewShaderPlatform != SP_NumPlatforms)
-	{
-		FShaderCompileUtilities::GenerateBrdfHeaders(MobilePreviewShaderPlatform);
-	}
-}
-#endif // WITH_EDITOR
 
 DECLARE_CYCLE_STAT(TEXT("FEngineLoop::PreInitPreStartupScreen.AfterStats"), STAT_FEngineLoop_PreInitPreStartupScreen_AfterStats, STATGROUP_LoadTime);
 DECLARE_CYCLE_STAT(TEXT("FEngineLoop::PreInitPostStartupScreen.AfterStats"), STAT_FEngineLoop_PreInitPostStartupScreen_AfterStats, STATGROUP_LoadTime);
@@ -2851,13 +2819,6 @@ PRAGMA_ENABLE_DEPRECATION_WARNINGS
 		RenderUtilsInit();
 	}
 
-#if WITH_EDITOR
-	{
-		// pre-generate certain shader code based on the engine configuration (editor only)
-		ShaderAutogenInit();
-	}
-#endif
-
 	{
 		bool bUseCodeLibrary = FPlatformProperties::RequiresCookedData() || GAllowCookedDataInEditorBuilds;
 		if (bUseCodeLibrary)
