Add AnimGraph module to be loaded during editor boot to avoid trying to load it from the async loading thread

diff --git a/Engine/Source/Runtime/Launch/Private/LaunchEngineLoop.cpp b/Engine/Source/Runtime/Launch/Private/LaunchEngineLoop.cpp
--- a/Engine/Source/Runtime/Launch/Private/LaunchEngineLoop.cpp
+++ b/Engine/Source/Runtime/Launch/Private/LaunchEngineLoop.cpp
@@ -4148,6 +4148,10 @@ bool FEngineLoop::LoadStartupCoreModules()
 	// Required during serialization of AnimSequence which could happen from async loading thread.
 	// See UAnimSequence::UpdateFrameRate().
 	FModuleManager::Get().LoadModule(TEXT("TimeManagement"));
+
+	// Required during construction of UAnimBlueprint which could happen from async loading thread.
+	// See UAnimBlueprint::UAnimBlueprint().
+	FModuleManager::Get().LoadModule(TEXT("AnimGraph"));
 #endif
 
 	FModuleManager::Get().LoadModule(TEXT("PacketHandler"));
