fixed deprecation warnings

diff --git a/Engine/Source/Runtime/Launch/Private/Android/LaunchAndroid.cpp b/Engine/Source/Runtime/Launch/Private/Android/LaunchAndroid.cpp
--- a/Engine/Source/Runtime/Launch/Private/Android/LaunchAndroid.cpp
+++ b/Engine/Source/Runtime/Launch/Private/Android/LaunchAndroid.cpp
@@ -517,12 +517,12 @@ int32 AndroidMain(struct android_app* state)
 		GAndroidWindowLock.Lock();
 	}
 
-	FDelegateHandle ConfigReadyHandle = FCoreDelegates::ConfigReadyForUse.AddStatic(&ApplyAndroidCompatConfigRules);
+	FDelegateHandle ConfigReadyHandle = FCoreDelegates::TSConfigReadyForUse().AddStatic(&ApplyAndroidCompatConfigRules);
 
 	// initialize the engine
 	int32 PreInitResult = GEngineLoop.PreInit(0, NULL, FCommandLine::Get());
 
-	FCoreDelegates::ConfigReadyForUse.Remove(ConfigReadyHandle);
+	FCoreDelegates::TSConfigReadyForUse().Remove(ConfigReadyHandle);
 
 	if (PreInitResult != 0)
 	{
