InitializeAfterProjectFilePath must be called for the entire hierarchy of platform files starting with FPlatformFileManager:Get().GetPlatformFile()

diff --git a/Engine/Source/Runtime/Launch/Private/LaunchEngineLoop.cpp b/Engine/Source/Runtime/Launch/Private/LaunchEngineLoop.cpp
--- a/Engine/Source/Runtime/Launch/Private/LaunchEngineLoop.cpp
+++ b/Engine/Source/Runtime/Launch/Private/LaunchEngineLoop.cpp
@@ -2200,7 +2200,11 @@ int32 FEngineLoop::PreInitPreStartupScreen(const TCHAR* CmdLine)
 #endif
 
 	// Initialize platform file with knowledge of the project file path before fixing the casing
-	FPlatformFileManager::Get().GetPlatformFile().InitializeAfterProjectFilePath();
+	IPlatformFile* CurrentPlatformFile = &FPlatformFileManager::Get().GetPlatformFile();
+	for (IPlatformFile* PlatformFileChainElement = CurrentPlatformFile; PlatformFileChainElement; PlatformFileChainElement = PlatformFileChainElement->GetLowerLevel())
+	{
+		PlatformFileChainElement->InitializeAfterProjectFilePath();
+	}
 
 #if !IS_PROGRAM
 	// Now let the platform file fix the project file path case before we attempt to fix the game name
