Fix server crash accesing the global SM.

diff --git a/Engine/Source/Runtime/Launch/Private/LaunchEngineLoop.cpp b/Engine/Source/Runtime/Launch/Private/LaunchEngineLoop.cpp
--- a/Engine/Source/Runtime/Launch/Private/LaunchEngineLoop.cpp
+++ b/Engine/Source/Runtime/Launch/Private/LaunchEngineLoop.cpp
@@ -3895,6 +3895,7 @@ int32 FEngineLoop::PreInitPostStartupScreen(const TCHAR* CmdLine)
 	}
 
 	// precache compute PSOs for global shaders - enough time should have passed since loading global SM to avoid a blocking load
+	if (AllowGlobalShaderLoad())
 	{
 		PrecacheComputePipelineStatesForGlobalShaders(GShaderPlatformForFeatureLevel[GMaxRHIFeatureLevel], nullptr);
 	}
