Fix duplicate loctext and init problems in commandlets (UE-165492).

diff --git a/Engine/Source/Runtime/Launch/Private/LaunchEngineLoop.cpp b/Engine/Source/Runtime/Launch/Private/LaunchEngineLoop.cpp
--- a/Engine/Source/Runtime/Launch/Private/LaunchEngineLoop.cpp
+++ b/Engine/Source/Runtime/Launch/Private/LaunchEngineLoop.cpp
@@ -6383,6 +6383,7 @@ void FEngineLoop::PostInitRHI()
 	}
 #endif //(!UE_BUILD_SHIPPING)
 
+	if (FApp::CanEverRender())
 	{
 		// perform an early check of hardware capabilities
 		EShaderPlatform ShaderPlatform = GMaxRHIShaderPlatform;
@@ -6394,7 +6395,7 @@ void FEngineLoop::PostInitRHI()
 		{
 			UE_LOG(LogInit, Log, TEXT("ShaderPlatform=%d RHISupportsInstancedStereo()=%d GRHISupportsArrayIndexFromAnyShader=%d"),
 				ShaderPlatform, RHISupportsMultiViewport(ShaderPlatform), GRHISupportsArrayIndexFromAnyShader);
-			FPlatformMisc::MessageBoxExt(EAppMsgType::Ok, *NSLOCTEXT("InstancedStereo", "UnableToUseInstancedStereoRendering", "Instanced Stereo cannot be used due to a missing functionality on the system. Please check log files for more info.").ToString(),
+			FPlatformMisc::MessageBoxExt(EAppMsgType::Ok, *NSLOCTEXT("InstancedStereo", "UnableToUseInstancedStereoRenderingText", "Instanced Stereo cannot be used due to a missing functionality on the system. Please check log files for more info.").ToString(),
 				*NSLOCTEXT("InstancedStereo", "UnableToUseInstancedStereoRendering", "Unable to use Instanced Stereo Rendering.").ToString());
 			FPlatformMisc::RequestExitWithStatus(true, 1);
 			// unreachable
