When refusing to run due to ISR, avoid getting stuck if the run is unattended.

- Should be handled in MessageBoxExt function itself but it is not yet implemented across all platforms.

diff --git a/Engine/Source/Runtime/Launch/Private/LaunchEngineLoop.cpp b/Engine/Source/Runtime/Launch/Private/LaunchEngineLoop.cpp
--- a/Engine/Source/Runtime/Launch/Private/LaunchEngineLoop.cpp
+++ b/Engine/Source/Runtime/Launch/Private/LaunchEngineLoop.cpp
@@ -6395,8 +6395,18 @@ void FEngineLoop::PostInitRHI()
 		{
 			UE_LOG(LogInit, Log, TEXT("ShaderPlatform=%d RHISupportsInstancedStereo()=%d GRHISupportsArrayIndexFromAnyShader=%d"),
 				ShaderPlatform, RHISupportsMultiViewport(ShaderPlatform), GRHISupportsArrayIndexFromAnyShader);
-			FPlatformMisc::MessageBoxExt(EAppMsgType::Ok, *NSLOCTEXT("InstancedStereo", "UnableToUseInstancedStereoRenderingText", "Instanced Stereo cannot be used due to a missing functionality on the system. Please check log files for more info.").ToString(),
-				*NSLOCTEXT("InstancedStereo", "UnableToUseInstancedStereoRendering", "Unable to use Instanced Stereo Rendering.").ToString());
+
+			// MessageBoxExt may not yet handle unattended runs itself (see UE-165694), so special-case here to avoid getting stuck in a commandlet with rendering enabled
+			const FText MessageText = NSLOCTEXT("InstancedStereo", "UnableToUseInstancedStereoRenderingText", "Cannot render an Instanced Stereo-enabled project due to a missing functionality on the system. Please check log files for more info.");
+			if (!FApp::IsUnattended())
+			{
+				FPlatformMisc::MessageBoxExt(EAppMsgType::Ok, *MessageText.ToString(),
+					*NSLOCTEXT("InstancedStereo", "UnableToUseInstancedStereoRendering", "Unable to use Instanced Stereo Rendering.").ToString());
+			}
+			else
+			{
+				UE_LOG(LogInit, Error, TEXT("%s"), *MessageText.ToString());					
+			}
 			FPlatformMisc::RequestExitWithStatus(true, 1);
 			// unreachable
 		}
