Print commandlet execution time in hms

diff --git a/Engine/Source/Runtime/Launch/Private/LaunchEngineLoop.cpp b/Engine/Source/Runtime/Launch/Private/LaunchEngineLoop.cpp
--- a/Engine/Source/Runtime/Launch/Private/LaunchEngineLoop.cpp
+++ b/Engine/Source/Runtime/Launch/Private/LaunchEngineLoop.cpp
@@ -3810,7 +3810,27 @@ int32 FEngineLoop::PreInitPostStartupScreen(const TCHAR* CmdLine)
 			}
 
 			double CommandletExecutionTime = FPlatformTime::Seconds() - CommandletExecutionStartTime;
-			UE_LOG(LogInit, Display, LINE_TERMINATOR TEXT( "Execution of commandlet took:  %.2f seconds"), CommandletExecutionTime );
+			if (CommandletExecutionTime <= 60)
+			{
+				UE_LOG(LogInit, Display, LINE_TERMINATOR TEXT("Execution of commandlet took:  %.2f seconds"), CommandletExecutionTime);
+			}
+			else
+			{
+				FTimespan ExecutionTimeSpan = FTimespan::FromSeconds(CommandletExecutionTime);
+				int32 Hours = (int32)(ExecutionTimeSpan.GetTotalHours());
+				int32 Minutes = ExecutionTimeSpan.GetMinutes();
+				int32 Seconds = ExecutionTimeSpan.GetSeconds();
+				// Tried FText::AsTimespan here but it actually felt harder to visually parse than just explicit labeling. Leaving the
+				// seconds in so that it's easy to subtract between multiple runs.
+				if (Hours)
+				{
+					UE_LOG(LogInit, Display, LINE_TERMINATOR TEXT("Execution of commandlet took:  %dh %dm %ds (%.2f seconds)"), Hours, Minutes, Seconds, CommandletExecutionTime);
+				}
+				else
+				{
+					UE_LOG(LogInit, Display, LINE_TERMINATOR TEXT("Execution of commandlet took:  %dm %ds (%.2f seconds)"), Minutes, Seconds, CommandletExecutionTime);
+				}
+			}
 
 			// We're ready to exit!
 			return ErrorLevel;
