Command line: EngineLoop: Reworked the way we handle command-line parameters during initialization to allow more flexibility (project/game name).

Previously, it was assumed that the first parameter was the project path (.uproject) or the game name. The next one was optionally a map override name or a commandlet. While this expectation makes sense, it was difficult to comply given that on some platforms and some launch modes we combine command-line parameters provided to the executable (e.g. from Visual Studio) with ones in uecommandline.txt script. We had some conflicting requirements and because of that certain launch scenarios didn't work.

Modified the code to look for the game name / project path withing all the parameters, regardless of the position. Other than that, I've tried to preserve the preexisting behavior. In a few places I've updated the code to allow other switches at arbitrary location withing the command line.

Testing done:
   * Editor: no parameters (project selector); running commandlets (short and long names and/or -run=) with different sets of parameters; running the editor with different parameters; running the game (-game)
   * Game (cpp + bp): running commandlets (short and long names) with different parameters; running the game with different combonations of project, map and other switches.

#rn fix core Allowed more flexibility when looking for the project name command-line parameter. Previously, it was expected to be always the first parameter. Now, we check all the provided parameters. It fixes some launch scenarios (e.g. launching blueprint-only projects with additional command-line parameters) that didn't work due to the fixed order in which we combine command-line parameters with the contents of uecommandline.txt script.

diff --git a/Engine/Source/Runtime/Launch/Private/LaunchEngineLoop.cpp b/Engine/Source/Runtime/Launch/Private/LaunchEngineLoop.cpp
--- a/Engine/Source/Runtime/Launch/Private/LaunchEngineLoop.cpp
+++ b/Engine/Source/Runtime/Launch/Private/LaunchEngineLoop.cpp
@@ -1909,29 +1909,28 @@ int32 FEngineLoop::PreInitPreStartupScreen(const TCHAR* CmdLine)
 	bool bIsRunningAsDedicatedServer = false;
 	bool bIsRegularClient = false;
 
-	const SIZE_T CommandLineSize = FCString::Strlen(CmdLine) + 1;
-	TCHAR* CommandLineCopy = new TCHAR[CommandLineSize];
-	FCString::Strcpy(CommandLineCopy, CommandLineSize, CmdLine);
-	const TCHAR* ParsedCmdLine = CommandLineCopy;
+	FString TokenToForward;
 
-	FString Token = FParse::Token(ParsedCmdLine, 0);
-	// trim any whitespace at edges of string - this can happen if the token was quoted with leading or trailing whitespace
-	// VC++ tends to do this in its "external tools" config
-	Token.TrimStartAndEndInline();
+	TArray<FString> TokenArray     = TokenizeCommandline(CmdLine);
+
+	EGameStringType GameStringType = EGameStringType::Unknown;
+	FString         GameString     = ExtractGameStringArgument(TokenArray, GameStringType);
 
 	auto IsModeSelected = [&bHasCommandletToken, &bHasEditorToken, &bIsRegularClient, &bIsRunningAsDedicatedServer]()
 	{
 		return bHasCommandletToken || bHasEditorToken || bIsRegularClient || bIsRunningAsDedicatedServer;
 	};
 #if UE_EDITOR || WITH_ENGINE || WITH_EDITOR
-	const TCHAR* CommandletCommandLine = nullptr;
-	auto SetIsRunningAsCommandlet = [&CommandletCommandLine, &ParsedCmdLine,
-		&bHasCommandletToken, &bIsRunningAsDedicatedServer, &bHasEditorToken, &bIsRegularClient, &IsModeSelected]
+	FString CommandletCommandLine;
+	auto SetIsRunningAsCommandlet = [&CommandletCommandLine, &TokenArray,
+		&bHasCommandletToken, &bIsRunningAsDedicatedServer, &bHasEditorToken, &bIsRegularClient, &IsModeSelected, &TokenToForward]
 		(FStringView CommandletName)
 	{
 		checkf(!IsModeSelected(), TEXT("SetIsRunningAsCommandlet should not be called after mode has been selected."));
 		bHasCommandletToken = true;
 
+		TokenToForward = CommandletName;
+
 		GIsClient = true;
 		GIsServer = true;
 #if WITH_EDITORONLY_DATA
@@ -1945,7 +1944,7 @@ int32 FEngineLoop::PreInitPreStartupScreen(const TCHAR* CmdLine)
 		}
 #endif
 
-		CommandletCommandLine = ParsedCmdLine;
+		CommandletCommandLine = FString::Join(TokenArray, TEXT(" "));
 #if WITH_EDITOR
 		if (CommandletName == TEXTVIEW("cookcommandlet"))
 		{
@@ -1961,11 +1960,25 @@ int32 FEngineLoop::PreInitPreStartupScreen(const TCHAR* CmdLine)
 		PRIVATE_GAllowCommandletAudio = FParse::Param(FCommandLine::Get(), TEXT("AllowCommandletAudio"));
 	};
 #endif // UE_EDITOR || WITH_ENGINE || WITH_EDITOR
-	auto SetIsRunningAsRegularClient = [&IsModeSelected, &bIsRegularClient]()
+	auto SetIsRunningAsRegularClient = [&IsModeSelected, &bIsRegularClient, &TokenArray, &TokenToForward]()
 	{
 		checkf(!IsModeSelected(), TEXT("SetIsRunningAsRegularClient should not be called after mode has been selected."));
 		bIsRegularClient = true;
 
+		// Take the first non-switch command-line parameter (i.e. one not starting with a dash) and remember
+		// it to check later as it could be a mistyped commandlet (short name).
+		const FString* NonSwitchParam = TokenArray.FindByPredicate(
+			[](const FString& Token)
+			{
+				return Token[0] != TCHAR('-');
+			}
+		);
+
+		if (NonSwitchParam)
+		{
+			TokenToForward = *NonSwitchParam;
+		}
+
 		GIsClient = true;
 		GIsServer = false;
 #if WITH_EDITORONLY_DATA
@@ -2009,8 +2022,8 @@ int32 FEngineLoop::PreInitPreStartupScreen(const TCHAR* CmdLine)
 	{
 		// Header generation is treated as a commandlet
 		checkf(!IsRunningDedicatedServer(), TEXT("GIsUCCMakeStandaloneHeaderGenerator is not supported in a ServerOnly configuration, or with -server on commandline."));
-		Token = TEXT("GIsUCCMakeStandaloneHeaderGenerator");
-		SetIsRunningAsCommandlet(Token);
+
+		SetIsRunningAsCommandlet(TEXT("GIsUCCMakeStandaloneHeaderGenerator"));
 	}
 	else
 #endif
@@ -2024,16 +2037,14 @@ int32 FEngineLoop::PreInitPreStartupScreen(const TCHAR* CmdLine)
 #if UE_EDITOR || WITH_ENGINE || WITH_EDITOR
 	TArray<FString> NonSwitchTokens;
 	TArray<FString> Switches;
-	UCommandlet::ParseCommandLine(CommandLineCopy, NonSwitchTokens, Switches);
+	UCommandlet::ParseCommandLine(CmdLine, NonSwitchTokens, Switches);
 	if (!IsModeSelected())
 	{
 		for (const FString& ParsedToken : NonSwitchTokens)
 		{
 			if (ParsedToken.EndsWith(TEXT("Commandlet")))
 			{
-				Token = ParsedToken;
-				Token.TrimStartAndEndInline();
-				SetIsRunningAsCommandlet(Token);
+				SetIsRunningAsCommandlet(ParsedToken.TrimStartAndEnd());
 				break;
 			}
 		}
@@ -2043,13 +2054,13 @@ int32 FEngineLoop::PreInitPreStartupScreen(const TCHAR* CmdLine)
 			{
 				if (ParsedSwitch.StartsWith(TEXT("RUN=")))
 				{
-					Token = ParsedSwitch.RightChop(4);
-					Token.TrimStartAndEndInline();
-					if (!Token.EndsWith(TEXT("Commandlet")))
+					FString LocalToken = ParsedSwitch.RightChop(4);
+					LocalToken.TrimStartAndEndInline();
+					if (!LocalToken.EndsWith(TEXT("Commandlet")))
 					{
-						Token += TEXT("Commandlet");
+						LocalToken += TEXT("Commandlet");
 					}
-					SetIsRunningAsCommandlet(Token);
+					SetIsRunningAsCommandlet(LocalToken);
 					break;
 				}
 			}
@@ -2061,37 +2072,21 @@ int32 FEngineLoop::PreInitPreStartupScreen(const TCHAR* CmdLine)
 	// In all other cases we should check for the first token being the Project Specifier
 	if (!bHasCommandletToken)
 	{
-		// Path returned by FPaths::GetProjectFilePath() is normalized, so may have symlinks and ~ resolved and may differ from the original path to .uproject passed in the command line
-		FString NormalizedToken = Token;
-		FPaths::NormalizeFilename(NormalizedToken);
-
-		const bool bFirstTokenIsGameName = (FApp::HasProjectName() && Token == FApp::GetProjectName());
-		const bool bFirstTokenIsGameProjectFilePath = (FPaths::IsProjectFilePathSet() && NormalizedToken == FPaths::GetProjectFilePath());
-		const bool bFirstTokenIsGameProjectFileShortName = (FPaths::IsProjectFilePathSet() && Token == FPaths::GetCleanFilename(FPaths::GetProjectFilePath()));
-
-		if (bFirstTokenIsGameName || bFirstTokenIsGameProjectFilePath || bFirstTokenIsGameProjectFileShortName)
+		if (GameStringType != EGameStringType::Unknown)
 		{
-			// The first item on command line was the game name, ignore it and get the next item for the Token
-			FString RemainingCommandline = ParsedCmdLine;
-			FCString::Strcpy(CommandLineCopy, CommandLineSize, *RemainingCommandline);
-			ParsedCmdLine = CommandLineCopy;
-
-			// Set a new command-line that doesn't include the game name as the first argument
-			FCommandLine::Set(ParsedCmdLine);
+			// Set a new command-line that doesn't include the game name as the first argument.
+			FCommandLine::Set(*FString::Join(TokenArray, TEXT(" ")));
 
-			Token = FParse::Token(ParsedCmdLine, 0);
-			Token.TrimStartInline();
-
-			// if the next token is a project file, then we skip it (which can happen on some platforms that combine
-			// commandlines... this handles extra .uprojects, but if you run with MyGame MyGame, we can't tell if
-			// the second MyGame is a map or not)
-			while (FPaths::GetExtension(Token) == FProjectDescriptor::GetExtension())
-			{
-				Token = FParse::Token(ParsedCmdLine, 0);
-				Token.TrimStartInline();
-			}
+			// Remove spurious project file tokens (which can happen on some platforms that combine commandlines).
+			// This handles extra .uprojects, but if you run with MyGame MyGame, we can't tell if the second MyGame is a map or not.
+			TokenArray.RemoveAll(
+				[](const FString& Token)
+				{
+					return Token[0] != TCHAR('-') && FPaths::GetExtension(Token) == FProjectDescriptor::GetExtension();
+				}
+			);
 
-			if (bFirstTokenIsGameProjectFilePath || bFirstTokenIsGameProjectFileShortName)
+			if (GameStringType == EGameStringType::ProjectPath || GameStringType == EGameStringType::ProjectShortName)
 			{
 				// Convert it to relative if possible...
 				FString RelativeGameProjectFilePath = FFileManagerGeneric::DefaultConvertToRelativePath(*FPaths::GetProjectFilePath());
@@ -2104,18 +2099,22 @@ int32 FEngineLoop::PreInitPreStartupScreen(const TCHAR* CmdLine)
 
 #if UE_EDITOR
 		// Handle the first token being '-game' or '-server'
-		if (Token == TEXT("-GAME") || Token == TEXT("-SERVER"))
+		if (!TokenArray.IsEmpty())
 		{
-			// This isn't necessarily pretty, but many requests have been made to allow
-			//   UnrealEditor.exe <GAMENAME> -game <map>
-			// or
-			//   UnrealEditor.exe <GAMENAME> -game 127.0.0.0
-			// We don't want to remove the -game from the commandline just yet in case
-			// we need it for something later. So, just move it to the end for now...
-			FString RemainingCommandline = ParsedCmdLine;
-			RemainingCommandline.TrimStartInline();
-			RemainingCommandline += FString::Printf(TEXT(" %s"), *Token);
-			FCommandLine::Set(*RemainingCommandline);
+			if (TokenArray[0] == TEXT("-GAME") || TokenArray[0] == TEXT("-SERVER"))
+			{
+				// This isn't necessarily pretty, but many requests have been made to allow
+				//   UnrealEditor.exe <GAMENAME> -game <map>
+				// or
+				//   UnrealEditor.exe <GAMENAME> -game 127.0.0.0
+				// We don't want to remove the -game from the commandline just yet in case
+				// we need it for something later. So, just move it to the end for now...
+				FString LocalToken = TokenArray[0];
+				TokenArray.Add(LocalToken);
+				TokenArray.RemoveAt(0);
+
+				FCommandLine::Set(*FString::Join(TokenArray, TEXT(" ")));
+			}
 		}
 #endif
 	}
@@ -2133,7 +2132,7 @@ int32 FEngineLoop::PreInitPreStartupScreen(const TCHAR* CmdLine)
 #elif WITH_ENGINE && WITH_EDITOR && WITH_EDITORONLY_DATA 
 		// If a non-editor target build w/ WITH_EDITOR and WITH_EDITORONLY_DATA, use the old token check...
 		//@todo. Is this something we need to support?
-		if (Token == TEXT("EDITOR"))
+		if (TokenArray.Contains(TEXT("EDITOR")))
 		{
 			SetIsRunningAsEditor();
 		}
@@ -2634,7 +2633,7 @@ PRAGMA_ENABLE_DEPRECATION_WARNINGS
 #if !UE_EDITOR
 	// UE_EDITOR doesn't care the meaning of the token except for a few special cases (FooCommandlet, -run=, -game, -server)
 	// But when running without UE_EDITOR, we look up the token as a class to see if it is a commandlet name prefix
-	bool bIsLateCommandletToken = false;
+	FString LateCommandletName;
 	if (!IsModeSelected())
 	{
 		//@hack: We need to set these before calling StaticLoadClass so all required data gets loaded for the commandlets.
@@ -2645,21 +2644,36 @@ PRAGMA_ENABLE_DEPRECATION_WARNINGS
 #endif // WITH_EDITOR
 		TGuardValue<bool> ScopedGIsRunningCommandlet(PRIVATE_GIsRunningCommandlet, true);
 
-		bool bIsPossiblyCommandletName = Token.Len() && !Token.Contains(TEXT("-"));
+		// Take the first non-switch command-line parameter (i.e. one not starting with a dash). We will check if it's possibly a commandlet.
+		const FString* NonSwitchParam = TokenArray.FindByPredicate(
+			[](const FString& Token)
+			{
+				return Token[0] != TCHAR('-');
+			}
+		);
+
+		bool bIsPossiblyCommandletName = NonSwitchParam != nullptr;
 		if (bIsPossiblyCommandletName)
 		{
-			UClass* TempCommandletClass = FindFirstObject<UClass>(*(Token + TEXT("Commandlet")), EFindFirstObjectOptions::None, ELogVerbosity::Warning, TEXT("Looking for commandlet class"));
+			FString CommandletName = *NonSwitchParam;
+
+			if (!CommandletName.EndsWith(TEXT("Commandlet")))
+			{
+				CommandletName += TEXT("Commandlet");
+			}
+
+			UClass* TempCommandletClass = FindFirstObject<UClass>(*CommandletName, EFindFirstObjectOptions::None, ELogVerbosity::Warning, TEXT("Looking for commandlet class"));
 			if (TempCommandletClass)
 			{
 				checkf(TempCommandletClass->IsChildOf(UCommandlet::StaticClass()), TEXT("It is not valid to have a class that ends with \"Commandlet\" that is not a UCommandlet subclass."));
-				Token += TEXT("Commandlet");
-				bIsLateCommandletToken = true;
+
+				LateCommandletName = CommandletName;
 			}
 		}
 	}
-	if (bIsLateCommandletToken)
+	if (!LateCommandletName.IsEmpty())
 	{
-		SetIsRunningAsCommandlet(Token);
+		SetIsRunningAsCommandlet(LateCommandletName);
 	}
 
 	// In non-UE_EDITOR configuration this is the point where know enough to finish the mode decision and decide between commandlet or client
@@ -3045,15 +3059,14 @@ PRAGMA_ENABLE_DEPRECATION_WARNINGS
 	PreInitContext.bDisableDisregardForGC = bDisableDisregardForGC;
 	PreInitContext.bIsRegularClient = bIsRegularClient;
 #endif // WITH_ENGINE
-	PreInitContext.bIsPossiblyUnrecognizedCommandlet = bIsRegularClient && Token.Len() && !Token.Contains(TEXT("-"));
+	PreInitContext.bIsPossiblyUnrecognizedCommandlet = bIsRegularClient && TokenToForward.Len() && !TokenToForward.Contains(TEXT("-"));
 	PRAGMA_DISABLE_DEPRECATION_WARNINGS;
 	PreInitContext.bTokenDoesNotHaveDash = PreInitContext.bIsPossiblyUnrecognizedCommandlet;
 	PRAGMA_ENABLE_DEPRECATION_WARNINGS;
-	PreInitContext.Token = Token;
+	PreInitContext.Token = TokenToForward;
 #if UE_EDITOR || WITH_ENGINE
 	PreInitContext.CommandletCommandLine = CommandletCommandLine;
 #endif // UE_EDITOR || WITH_ENGINE
-	PreInitContext.CommandLineCopy = CommandLineCopy;
 	
 	return 0;
 }
@@ -3083,9 +3096,8 @@ int32 FEngineLoop::PreInitPostStartupScreen(const TCHAR* CmdLine)
 	bool bIsPossiblyUnrecognizedCommandlet = PreInitContext.bIsPossiblyUnrecognizedCommandlet;
 	FString Token = PreInitContext.Token;
 #if UE_EDITOR || WITH_ENGINE
-	const TCHAR* CommandletCommandLine = PreInitContext.CommandletCommandLine;
+	const TCHAR* CommandletCommandLine = *PreInitContext.CommandletCommandLine;
 #endif // UE_EDITOR || WITH_ENGINE
-	TCHAR* CommandLineCopy = PreInitContext.CommandLineCopy;
 	FScopedSlowTask& SlowTask = *PreInitContext.SlowTaskPtr;
 
 #if WITH_ENGINE
@@ -4046,6 +4058,67 @@ void FEngineLoop::LoadPreInitModules()
 #endif // WITH_EDITOR
 }
 
+TArray<FString> FEngineLoop::TokenizeCommandline(const TCHAR* CmdLine)
+{
+	TArray<FString> TokenArray;
+
+	const TCHAR* ParsedCmdLine = CmdLine;
+
+	while (*ParsedCmdLine)
+	{
+		FString Token = FParse::Token(ParsedCmdLine, 0);
+
+		Token.TrimStartAndEndInline();
+
+		if (!Token.IsEmpty())
+		{
+			TokenArray.Add(MoveTemp(Token));
+		}
+	}
+
+	return MoveTemp(TokenArray);
+}
+
+FString FEngineLoop::ExtractGameStringArgument(TArray<FString>& TokenArray, EGameStringType& OutStringType)
+{
+	for (int32 I = 0; I < TokenArray.Num(); ++I)
+	{
+		FString NormalizedToken = TokenArray[I];
+
+		// Path returned by FPaths::GetProjectFilePath() is normalized, so may have symlinks and ~ resolved and may differ from the original path to .uproject passed in the command line
+		FPaths::NormalizeFilename(NormalizedToken);
+
+		const bool bTokenIsGameName                 = (FApp::HasProjectName()         && TokenArray[I]   == FApp::GetProjectName());
+		const bool bTokenIsGameProjectFilePath      = (FPaths::IsProjectFilePathSet() && NormalizedToken == FPaths::GetProjectFilePath());
+		const bool bTokenIsGameProjectFileShortName = (FPaths::IsProjectFilePathSet() && TokenArray[I]   == FPaths::GetCleanFilename(FPaths::GetProjectFilePath()));
+
+		if (bTokenIsGameName || bTokenIsGameProjectFilePath || bTokenIsGameProjectFileShortName)
+		{
+			if (bTokenIsGameName)
+			{
+				OutStringType = EGameStringType::GameName;
+			}
+			else if (bTokenIsGameProjectFilePath)
+			{
+				OutStringType = EGameStringType::ProjectPath;
+			}
+			else if (bTokenIsGameProjectFileShortName)
+			{
+				OutStringType = EGameStringType::ProjectShortName;
+			}
+
+			FString Result = TokenArray[I];
+
+			TokenArray.RemoveAt(I);
+
+			return Result;
+		}
+	}
+
+	OutStringType = EGameStringType::Unknown;
+
+	return FString();
+}
 
 #if WITH_ENGINE
 
@@ -6495,10 +6568,6 @@ void FPreInitContext::Cleanup()
 #if WITH_ENGINE && !UE_SERVER
 	SlateRenderer = nullptr;
 #endif // WITH_ENGINE && !UE_SERVER
-	CommandletCommandLine = nullptr;
-
-	delete[] CommandLineCopy;
-	CommandLineCopy = nullptr;
 
 	delete SlowTaskPtr;
 	SlowTaskPtr = nullptr;
diff --git a/Engine/Source/Runtime/Launch/Public/LaunchEngineLoop.h b/Engine/Source/Runtime/Launch/Public/LaunchEngineLoop.h
--- a/Engine/Source/Runtime/Launch/Public/LaunchEngineLoop.h
+++ b/Engine/Source/Runtime/Launch/Public/LaunchEngineLoop.h
@@ -29,8 +29,7 @@ struct FPreInitContext
 	bool bIsPossiblyUnrecognizedCommandlet = false;
 
 	FString Token;
-	const TCHAR* CommandletCommandLine = nullptr;
-	TCHAR* CommandLineCopy = nullptr;
+	FString CommandletCommandLine;
 
 	FScopedSlowTask* SlowTaskPtr = nullptr;
 
@@ -179,6 +178,27 @@ protected:
 
 private:
 
+	/** Enumeration representing the type of the command-line argument representing the game (typically the first argument). */
+	enum class EGameStringType
+	{
+		GameName,
+		ProjectPath,
+		ProjectShortName,
+		Unknown
+	};
+
+	/** Takes a command-line string and returns an array of tokens (splits arguments using whitespaces). */
+	static TArray<FString> TokenizeCommandline(const TCHAR* CmdLine);
+
+	/**
+	  * Finds a command-line argument representing the game, removes it from the array and returns.
+	  * 
+	  * @param TokenArray    Array of the tokenized command line
+	  * @param OutStringType Type of the game string found.
+	  * @return String representing the game command-line parameter; empty string if not found (OutStringType == Unknown in such case).
+	  */
+	static FString ExtractGameStringArgument(TArray<FString>& TokenArray, EGameStringType& OutStringType);
+
 #if WITH_ENGINE
 
 	/** Holds the engine service. */
