Logging: Start the dedicated thread later to fix missing logs during early initialization

Starting the thread allows it to flush buffered lines, which is problematic if it happens before the default output devices have been initialized. Wait until after they are initialized by AppInit before starting the dedicated thread.

#rnx

diff --git a/Engine/Source/Runtime/Launch/Private/LaunchEngineLoop.cpp b/Engine/Source/Runtime/Launch/Private/LaunchEngineLoop.cpp
--- a/Engine/Source/Runtime/Launch/Private/LaunchEngineLoop.cpp
+++ b/Engine/Source/Runtime/Launch/Private/LaunchEngineLoop.cpp
@@ -1677,12 +1677,6 @@ int32 FEngineLoop::PreInitPreStartupScreen(const TCHAR* CmdLine)
 	// as soon as we determine whether GIsEditor == false
 	GLog->EnableBacklog(true);
 
-	// Try to start the dedicated primary thread now that the command line is available.
-	if (!FParse::Param(FCommandLine::Get(), TEXT("NoLogThread")))
-	{
-		GLog->TryStartDedicatedPrimaryThread();
-	}
-
 	// Initialize std out device as early as possible if requested in the command line
 #if PLATFORM_DESKTOP
 	// consoles don't typically have stdout, and FOutputDeviceDebug is responsible for echoing logs to the
@@ -2495,6 +2489,14 @@ PRAGMA_ENABLE_DEPRECATION_WARNINGS
 		}
 	}
 
+	// Try to start the dedicated primary thread now that the command line is available,
+	// and the default output devices have been created. Initializing earlier will cause
+	// logs to be missed by the default output devices.
+	if (!FParse::Param(FCommandLine::Get(), TEXT("NoLogThread")))
+	{
+		GLog->TryStartDedicatedPrimaryThread();
+	}
+
 	if (FPlatformProcess::SupportsMultithreading())
 	{
 		{
