Fixed: Reported memory leaks on Android Vulkan
Android had UObject system teardown disabled

diff --git a/Engine/Source/Runtime/Launch/Private/LaunchEngineLoop.cpp b/Engine/Source/Runtime/Launch/Private/LaunchEngineLoop.cpp
--- a/Engine/Source/Runtime/Launch/Private/LaunchEngineLoop.cpp
+++ b/Engine/Source/Runtime/Launch/Private/LaunchEngineLoop.cpp
@@ -4688,11 +4688,9 @@ void FEngineLoop::Exit()
 #endif // WITH_EDITOR
 	FModuleManager::Get().UnloadModule("WorldBrowser", true);
 
-#if !PLATFORM_ANDROID 	// AppPreExit doesn't work on Android
 	AppPreExit();
 
 	TermGamePhys();
-#endif
 
 #if WITH_COREUOBJECT
 	// PackageResourceManager depends on AssetRegistry, so must be shutdown before we unload the AssetRegistry module
@@ -4700,11 +4698,6 @@ void FEngineLoop::Exit()
 #endif
 	FModuleManager::Get().UnloadModule("AssetRegistry", true);
 
-#if PLATFORM_ANDROID
-	// AppPreExit() stops malloc profiler, do it here instead
-	MALLOC_PROFILER( GMalloc->Exec(nullptr, TEXT("MPROF STOP"), *GLog);	);
-#endif // !ANDROID
-
 #if WITH_EDITOR
 	IBulkDataRegistry::Shutdown();
 #endif
@@ -4732,7 +4725,6 @@ void FEngineLoop::Exit()
 	UE::Virtualization::Shutdown();
 #endif //WITH_EDITOR
 
-#if !PLATFORM_ANDROID // UnloadModules doesn't work on Android
 #if WITH_ENGINE
 	// Save the hot reload state
 	IHotReloadInterface* HotReload = IHotReloadInterface::GetPtr();
@@ -4749,7 +4741,6 @@ void FEngineLoop::Exit()
 		SCOPED_BOOT_TIMING("UnloadModulesAtShutdown");
 		FModuleManager::Get().UnloadModulesAtShutdown();
 	}
-#endif // !ANDROID
 
 	IStreamingManager::Shutdown();
 
