Change when ImageNtHeader is called to PreInitPostStartupScreen. This is after some hidden DbgHelp initialization is done so future calls to SymFromAddr succeed. There is not a clear answer for when ImageNtHeader can be called but where it was before was too early.
This is potentially an issue with large PDBs
This issue only was seen to reproduce in monolithic exes

diff --git a/Engine/Source/Runtime/Launch/Private/Launch.cpp b/Engine/Source/Runtime/Launch/Private/Launch.cpp
--- a/Engine/Source/Runtime/Launch/Private/Launch.cpp
+++ b/Engine/Source/Runtime/Launch/Private/Launch.cpp
@@ -19,9 +19,6 @@
 #endif
 #if PLATFORM_WINDOWS
 	#include "Windows/WindowsHWrapper.h"
-	#include "Windows/AllowWindowsPlatformTypes.h"
-	#include <DbgHelp.h>
-	#include "Windows/HideWindowsPlatformTypes.h"
 #endif
 
 
@@ -30,8 +27,6 @@ IMPLEMENT_MODULE(FDefaultModuleImpl, Launch);
 #if PLATFORM_WINDOWS || PLATFORM_MAC || PLATFORM_UNIX || PLATFORM_USE_GENERIC_LAUNCH_IMPLEMENTATION
 
 FEngineLoop	GEngineLoop;
-bool GIsConsoleExecutable = false;
-
 
 extern "C" int test_main(int argc, char ** argp)
 {
@@ -139,15 +134,6 @@ int32 GuardedMain( const TCHAR* CmdLine )
 	// That will also use the user folder for installed builds so we don't write into program files or whatever.
 #if PLATFORM_WINDOWS
 	FCString::Strcpy(MiniDumpFilenameW, *FString::Printf(TEXT("unreal-v%i-%s.dmp"), FEngineVersion::Current().GetChangelist(), *FDateTime::Now().ToString()));
-
-	if (PIMAGE_NT_HEADERS NtHeaders = ImageNtHeader(GetModuleHandle(nullptr)))
-	{
-		GIsConsoleExecutable = (NtHeaders->OptionalHeader.Subsystem == IMAGE_SUBSYSTEM_WINDOWS_CUI);
-	}
-	else
-	{
-		GIsConsoleExecutable = (GetFileType(GetStdHandle(STD_OUTPUT_HANDLE)) == FILE_TYPE_CHAR);
-	}
 #endif
 
 	FTrackedActivity::GetEngineActivity().Update(TEXT("Initializing"));
diff --git a/Engine/Source/Runtime/Launch/Private/LaunchEngineLoop.cpp b/Engine/Source/Runtime/Launch/Private/LaunchEngineLoop.cpp
--- a/Engine/Source/Runtime/Launch/Private/LaunchEngineLoop.cpp
+++ b/Engine/Source/Runtime/Launch/Private/LaunchEngineLoop.cpp
@@ -239,6 +239,7 @@ class FFeedbackContext;
 	#include "Windows/AllowWindowsPlatformTypes.h"
 	#include <ObjBase.h>
 	#include "Windows/HideWindowsPlatformTypes.h"
+	#include <DbgHelp.h>
 #endif
 
 #if WITH_ENGINE
@@ -280,6 +281,8 @@ class FFeedbackContext;
 #include "IOS/IOSAppDelegate.h"
 #endif
 
+bool GIsConsoleExecutable = false;
+
 int32 GUseDisregardForGCOnDedicatedServers = 1;
 static FAutoConsoleVariableRef CVarUseDisregardForGCOnDedicatedServers(
 	TEXT("gc.UseDisregardForGCOnDedicatedServers"),
@@ -3084,6 +3087,23 @@ int32 FEngineLoop::PreInitPostStartupScreen(const TCHAR* CmdLine)
 		return 0;
 	}
 
+	extern bool GIsConsoleExecutable;
+#if PLATFORM_WINDOWS
+	/*
+	Note - ImageNtHeader must be called after DbgHelp is initialized.Failure to wait before initialization will cause calls to
+	SymFromAddr for monolithic exes to fail. This results in callstacks not being written to the log.
+	From discussion with Microsoft it is unclear when this time is and could be related to large pdbs. At this time here is an appropriate spot.
+	*/
+	if (PIMAGE_NT_HEADERS NtHeaders = ImageNtHeader(GetModuleHandle(nullptr)))
+	{
+		GIsConsoleExecutable = (NtHeaders->OptionalHeader.Subsystem == IMAGE_SUBSYSTEM_WINDOWS_CUI);
+	}
+	else
+	{
+		GIsConsoleExecutable = (GetFileType(GetStdHandle(STD_OUTPUT_HANDLE)) == FILE_TYPE_CHAR);
+	}
+#endif // PLATFORM_WINDOWS
+
 	FScopeCycleCounter CycleCount_AfterStats(GET_STATID(STAT_FEngineLoop_PreInitPostStartupScreen_AfterStats));
 
 	// Restore PreInitContext
@@ -3592,7 +3612,6 @@ int32 FEngineLoop::PreInitPostStartupScreen(const TCHAR* CmdLine)
 			}
 
 #if PLATFORM_WINDOWS || PLATFORM_MAC || PLATFORM_UNIX
-			extern bool GIsConsoleExecutable;
 			if (GIsConsoleExecutable)
 			{
 				if (GLogConsole != nullptr && GLogConsole->IsAttached())
