Fix running commandlets without a project.

diff --git a/Engine/Source/Runtime/Launch/Private/LaunchEngineLoop.cpp b/Engine/Source/Runtime/Launch/Private/LaunchEngineLoop.cpp
--- a/Engine/Source/Runtime/Launch/Private/LaunchEngineLoop.cpp
+++ b/Engine/Source/Runtime/Launch/Private/LaunchEngineLoop.cpp
@@ -557,6 +557,12 @@ bool ParseGameProjectFromCommandLine(const TCHAR* InCmdLine, FString& OutProject
 				return true;
 			}
 		}
+		else if (CommandLineToken.StartsWith(TEXT("-run=")))
+		{
+			// We encountered a commandlet invocation. Stop checking for the project name as commandlets
+			// can have arguments to them that don't start with a `-`
+			break;
+		}
 		CommandLineToken = FParse::Token(CmdLine, 0);
 	}
 
