Gauntlet - project name that's generated into UECommandline.txt is never used
Project name should be passed as first argument all times.
Revertws CL 21706201 which was more of a patch than a real solution to the problem

diff --git a/Engine/Source/Programs/AutomationTool/AutomationUtils/DeploymentContext.cs b/Engine/Source/Programs/AutomationTool/AutomationUtils/DeploymentContext.cs
--- a/Engine/Source/Programs/AutomationTool/AutomationUtils/DeploymentContext.cs
+++ b/Engine/Source/Programs/AutomationTool/AutomationUtils/DeploymentContext.cs
@@ -496,7 +496,7 @@ public class DeploymentContext //: ProjectParams
 
 		RelativeProjectRootForStage = new StagedDirectoryReference(ShortProjectName);
 
-		ProjectArgForCommandLines = CommandUtils.MakePathSafeToUseWithCommandLine(RawProjectPath.FullName);
+		ProjectArgForCommandLines = string.Format("-project={0}", CommandUtils.MakePathSafeToUseWithCommandLine(RawProjectPath.FullName));
 		CookSourceRuntimeRootDir = RuntimeRootDir = LocalRoot;
 		RuntimeProjectRootDir = ProjectRoot;
 
@@ -529,7 +529,7 @@ public class DeploymentContext //: ProjectParams
 			RuntimeRootDir = StageDirectory;
 			CookSourceRuntimeRootDir = DirectoryReference.Combine(BaseStageDirectory, CookPlatform);
 			RuntimeProjectRootDir = DirectoryReference.Combine(StageDirectory, RelativeProjectRootForStage.Name);
-			ProjectArgForCommandLines = CommandUtils.MakePathSafeToUseWithCommandLine(UProjectCommandLineArgInternalRoot + RelativeProjectRootForStage.Name + "/" + ShortProjectName + ".uproject");
+			ProjectArgForCommandLines = string.Format("-project={0}", CommandUtils.MakePathSafeToUseWithCommandLine(UProjectCommandLineArgInternalRoot + RelativeProjectRootForStage.Name + "/" + ShortProjectName + ".uproject"));
 		}
 		if (Archive)
 		{
diff --git a/Engine/Source/Runtime/Launch/Private/LaunchEngineLoop.cpp b/Engine/Source/Runtime/Launch/Private/LaunchEngineLoop.cpp
--- a/Engine/Source/Runtime/Launch/Private/LaunchEngineLoop.cpp
+++ b/Engine/Source/Runtime/Launch/Private/LaunchEngineLoop.cpp
@@ -522,48 +522,39 @@ void InitializeStdOutDevice()
 
 bool ParseGameProjectFromCommandLine(const TCHAR* InCmdLine, FString& OutProjectFilePath, FString& OutGameName)
 {
-	OutProjectFilePath = TEXT("");
-	OutGameName = TEXT("");
+	FString ProjectFilePathOption;
 
 	const TCHAR *CmdLine = InCmdLine;
-	FString CommandLineToken = FParse::Token(CmdLine, 0);
+	bool ParseProjectOptionResult = FParse::Value(CmdLine, TEXT("project="), ProjectFilePathOption);
+	FString FirstCommandLineToken = FParse::Token(CmdLine, 0);
 
-	// Avoid checking too many arguments
-	int MaxArgsChecked = 10;
+	// trim any whitespace at edges of string - this can happen if the token was quoted with leading or trailing whitespace
+	// VC++ tends to do this in its "external tools" config
+	FirstCommandLineToken.TrimStartInline();
 
-	while (!CommandLineToken.IsEmpty() && MaxArgsChecked-- > 0)
-	{
-		// trim any whitespace at edges of string - this can happen if the token was quoted with leading or trailing whitespace
-		// VC++ tends to do this in its "external tools" config
-		CommandLineToken.TrimStartInline();
+	// Output parameters
+	OutProjectFilePath = TEXT("");
+	OutGameName = TEXT("");
 
-		if (CommandLineToken.Len() && !CommandLineToken.StartsWith(TEXT("-")))
+	if ( ( FirstCommandLineToken.Len() && !FirstCommandLineToken.StartsWith(TEXT("-"))) || ParseProjectOptionResult)
+	{
+		// The first command line argument could be the project file if it exists or the game name if not launching with a project file
+		FString ProjectFilePath = ProjectFilePathOption.IsEmpty() ? FString(FirstCommandLineToken) : ProjectFilePathOption;
+		if ( FPaths::GetExtension(ProjectFilePath) == FProjectDescriptor::GetExtension() )
 		{
-			// This command line argument could be the project file if it exists or the game name if not launching with a project file
-			const FString ProjectFilePath = FString(CommandLineToken);
-			if (FPaths::GetExtension(ProjectFilePath) == FProjectDescriptor::GetExtension())
-			{
-				OutProjectFilePath = CommandLineToken;
-				// Here we derive the game name from the project file
-				OutGameName = FPaths::GetBaseFilename(OutProjectFilePath);
-				return true;
-			}
-			else if (FPaths::IsRelative(CommandLineToken) && FPlatformProperties::IsMonolithicBuild() == false)
-			{
-				// Full game name is assumed to be the first token
-				OutGameName = MoveTemp(CommandLineToken);
-				// Derive the project path from the game name. All games must have a uproject file, even if they are in the root folder.
-				OutProjectFilePath = FPaths::Combine(*FPaths::RootDir(), *OutGameName, *FString(OutGameName + TEXT(".") + FProjectDescriptor::GetExtension()));
-				return true;
-			}
+			OutProjectFilePath = ProjectFilePath;
+			// Here we derive the game name from the project file
+			OutGameName = FPaths::GetBaseFilename(OutProjectFilePath);
+			return true;
 		}
-		else if (CommandLineToken.StartsWith(TEXT("-run=")))
+		else if (FPaths::IsRelative(ProjectFilePath) && FPlatformProperties::IsMonolithicBuild() == false)
 		{
-			// We encountered a commandlet invocation. Stop checking for the project name as commandlets
-			// can have arguments to them that don't start with a `-`
-			break;
+			// Full game name is assumed to be the first token
+			OutGameName = MoveTemp(ProjectFilePath);
+			// Derive the project path from the game name. All games must have a uproject file, even if they are in the root folder.
+			OutProjectFilePath = FPaths::Combine(*FPaths::RootDir(), *OutGameName, *FString(OutGameName + TEXT(".") + FProjectDescriptor::GetExtension()));
+			return true;
 		}
-		CommandLineToken = FParse::Token(CmdLine, 0);
 	}
 
 #if WITH_EDITOR
