Ensure on cook errors in vk.

diff --git a/Engine/Source/Runtime/Launch/Private/LaunchEngineLoop.cpp b/Engine/Source/Runtime/Launch/Private/LaunchEngineLoop.cpp
--- a/Engine/Source/Runtime/Launch/Private/LaunchEngineLoop.cpp
+++ b/Engine/Source/Runtime/Launch/Private/LaunchEngineLoop.cpp
@@ -3760,6 +3760,13 @@ int32 FEngineLoop::PreInitPostStartupScreen(const TCHAR* CmdLine)
 				GWarn->GetErrors(AllErrors);
 				GWarn->GetWarnings(AllWarnings);
 
+				bool bEnsureOnError = false;
+				GConfig->GetBool(TEXT("Core.System"), TEXT("EnsureCommandletOnError"), bEnsureOnError, GEngineIni);
+				if (bEnsureOnError)
+				{
+					ensure(AllErrors.Num() == 0);
+				}
+
 				if (AllErrors.Num() || AllWarnings.Num())
 				{
 					SET_WARN_COLOR(COLOR_WHITE);
