Logging: Added -JsonStdOut param to write structured log output to stdout

diff --git a/Engine/Source/Runtime/Core/Private/Logging/StructuredLog.cpp b/Engine/Source/Runtime/Core/Private/Logging/StructuredLog.cpp
--- a/Engine/Source/Runtime/Core/Private/Logging/StructuredLog.cpp
+++ b/Engine/Source/Runtime/Core/Private/Logging/StructuredLog.cpp
@@ -189,8 +189,8 @@ FLogTemplate* FLogTemplate::Create(const TCHAR* Format, const FLogField* Fields,
 	checkf(bPositional || Algo::AllOf(FieldsView, &FLogField::Name),
 		TEXT("Log fields must be entirely named or entirely anonymous. [[%s]]"), Format);
 	checkf(bPositional || Algo::AllOf(FieldsView,
-		[](const FLogField& Field) { return *Field.Name && FAsciiSet::HasOnly(Field.Name, ValidLogFieldName); }),
-		TEXT("Log field names must match \"[A-Za-z0-9_]+\" in [[%s]]."), Format);
+		[](const FLogField& Field) { return *Field.Name && *Field.Name != '_' && FAsciiSet::HasOnly(Field.Name, ValidLogFieldName); }),
+		TEXT("Log field names must match \"[A-Za-z0-9][A-Za-z0-9_]*\" in [[%s]]."), Format);
 
 	TArray<FLogTemplateOp, TInlineAllocator<16>> Ops;
 
@@ -244,10 +244,13 @@ FLogTemplate* FLogTemplate::Create(const TCHAR* Format, const FLogField* Fields,
 		const TCHAR* const FieldStart = TextStart;
 		const TCHAR* const FieldNameEnd = FAsciiSet::Skip(FieldStart + 1, ValidLogFieldName);
 		checkf(*FieldNameEnd, TEXT("Log format has an unterminated field reference. Use '{{' to escape '{' if needed. [[%s]]"), Format);
-		checkf(*FieldNameEnd == TEXT('}'), TEXT("Log format has invalid character '%c' in field name. Use '{{' to escape '{' if needed. Names must match \"[A-Za-z0-9_]+\". [[%s]]"), *FieldNameEnd, Format);
+		checkf(*FieldNameEnd == TEXT('}'), TEXT("Log format has invalid character '%c' in field name. "
+			"Use '{{' to escape '{' if needed. Names must match \"[A-Za-z0-9][A-Za-z0-9_]*\". [[%s]]"), *FieldNameEnd, Format);
 		const TCHAR* const FieldEnd = FieldNameEnd + 1;
 		const int32 FieldLen = UE_PTRDIFF_TO_INT32(FieldEnd - FieldStart);
 		const int32 FieldNameLen = FieldLen - 2;
+		checkf(FieldStart[1] != TEXT('_'), TEXT("Log format uses reserved field name '%.*s' with leading '_'. "
+			"Names must match \"[A-Za-z0-9][A-Za-z0-9_]*\". [[%s]]"), FieldNameLen, FieldStart + 1, Format);
 
 		if (bFindFields && !bPositional)
 		{
@@ -589,7 +592,6 @@ void LogWithFieldArray(const FStaticLogRecord& Log, const FLogField* Fields, con
 	}
 #endif
 
-	FOutputDevice* LogOverride = nullptr;
 	switch (Log.Verbosity)
 	{
 	case ELogVerbosity::Error:
diff --git a/Engine/Source/Runtime/Launch/Private/LaunchEngineLoop.cpp b/Engine/Source/Runtime/Launch/Private/LaunchEngineLoop.cpp
--- a/Engine/Source/Runtime/Launch/Private/LaunchEngineLoop.cpp
+++ b/Engine/Source/Runtime/Launch/Private/LaunchEngineLoop.cpp
@@ -15,9 +15,11 @@
 #include "Misc/FileHelper.h"
 #include "Internationalization/TextLocalizationManagerGlobals.h"
 #include "Logging/LogSuppressionInterface.h"
+#include "Logging/StructuredLog.h"
 #include "Async/TaskGraphInterfaces.h"
 #include "Async/Fundamental/Scheduler.h"
 #include "MemPro/MemProProfiler.h"
+#include "Misc/AsciiSet.h"
 #include "Misc/TimeGuard.h"
 #include "Misc/Paths.h"
 #include "Misc/PathViews.h"
@@ -31,12 +33,16 @@
 #include "Misc/App.h"
 #include "Misc/OutputDeviceConsole.h"
 #include "Misc/ScopeExit.h"
+#include "Misc/StringBuilder.h"
 #include "Misc/TrackedActivity.h"
 #include "HAL/PlatformFileManager.h"
 #include "HAL/PlatformMemoryHelpers.h"
 #include "HAL/FileManagerGeneric.h"
 #include "HAL/ExceptionHandling.h"
 #include "HAL/IPlatformFileManagedStorageWrapper.h"
+#include "Serialization/CompactBinary.h"
+#include "Serialization/CompactBinarySerialization.h"
+#include "Serialization/CompactBinaryWriter.h"
 #include "Stats/StatsMallocProfilerProxy.h"
 #include "Trace/Trace.inl"
 #include "ProfilingDebugging/TraceAuxiliary.h"
@@ -357,7 +363,7 @@ private:
 
 // Pipe output to std output
 // This enables UBT to collect the output for it's own use
-class FOutputDeviceStdOutput : public FOutputDevice
+class FOutputDeviceStdOutput final : public FOutputDevice
 {
 public:
 	FOutputDeviceStdOutput()
@@ -366,6 +372,8 @@ public:
 		bIsConsoleOutput = IsStdoutAttachedToConsole() && !FParse::Param(FCommandLine::Get(), TEXT("GenericConsoleOutput"));
 #endif
 
+		bIsJsonOutput = FParse::Param(FCommandLine::Get(), TEXT("JsonStdOut"));
+
 		if (FParse::Param(FCommandLine::Get(), TEXT("AllowStdOutLogVerbosity")))
 		{
 			AllowedLogVerbosity = ELogVerbosity::Log;
@@ -382,54 +390,294 @@ public:
 		}
 	}
 
-	virtual bool CanBeUsedOnAnyThread() const override
+	bool CanBeUsedOnAnyThread() const final { return true; }
+	bool CanBeUsedOnPanicThread() const final { return true; }
+
+	void Serialize(const TCHAR* V, ELogVerbosity::Type Verbosity, const FName& Category) final
 	{
-		return true;
+		Serialize(V, Verbosity, Category, -1.0);
 	}
 
-	virtual bool CanBeUsedOnPanicThread() const override
+	void Serialize(const TCHAR* V, ELogVerbosity::Type Verbosity, const FName& Category, double Time) final
 	{
-		return true;
+		if (Verbosity <= AllowedLogVerbosity)
+		{
+			return bIsJsonOutput ? SerializeAsJson(V, Verbosity, Category, Time) : SerializeAsText(V, Verbosity, Category, Time);
+		}
 	}
 
-	virtual void Serialize( const TCHAR* V, ELogVerbosity::Type Verbosity, const class FName& Category ) override
+	void SerializeRecord(const UE::FLogRecord& Record) final
 	{
-		if (Verbosity <= AllowedLogVerbosity)
+		if (Record.GetVerbosity() <= AllowedLogVerbosity)
 		{
-			FString line = FOutputDeviceHelper::FormatLogLine(Verbosity, Category, V, GPrintLogTimes);
+			return bIsJsonOutput ? SerializeRecordAsJson(Record) : SerializeRecordAsText(Record);
+		}
+	}
 
-#if PLATFORM_WINDOWS
-			if (bIsConsoleOutput)
+private:
+	// Several functions below are FORCENOINLINE to reduce total required stack space by limiting
+	// the scope of string builders and compact binary writers.
+
+	FORCENOINLINE void SerializeAsText(const TCHAR* V, ELogVerbosity::Type Verbosity, const FName& Category, double Time)
+	{
+	#if PLATFORM_WINDOWS
+		if (bIsConsoleOutput)
+		{
+			return WriteLine<WIDECHAR>(V, Verbosity, Category, Time);
+		}
+	#endif
+	#if PLATFORM_TCHAR_IS_UTF8CHAR || PLATFORM_TCHAR_IS_CHAR16
+		WriteLine<UTF8CHAR>(V, Verbosity, Category, Time);
+	#else
+		WriteLine<WIDECHAR>(V, Verbosity, Category, Time);
+	#endif
+	}
+
+	FORCENOINLINE void SerializeRecordAsText(const UE::FLogRecord& Record)
+	{
+		TStringBuilder<512> V;
+		Record.FormatMessageTo(V);
+		Serialize(*V, Record.GetVerbosity(), Record.GetCategory(), FPlatformTime::ToSeconds64(Record.GetTime().GetCycles()));
+	}
+
+	template <typename CharType>
+	FORCENOINLINE void WriteLine(const TCHAR* V, ELogVerbosity::Type Verbosity, const FName& Category, double Time)
+	{
+		TStringBuilderWithBuffer<CharType, 512> Line;
+		FOutputDeviceHelper::AppendFormatLogLine(Line, Verbosity, Category, V, GPrintLogTimes, Time);
+		Line.AppendChar('\n');
+		WriteLine(Line);
+	}
+
+	FORCENOINLINE void SerializeAsJson(const TCHAR* V, ELogVerbosity::Type Verbosity, const FName& Category, double Time)
+	{
+		const bool bShowCategory = GPrintLogCategory && !Category.IsNone();
+		const bool bShowVerbosity = GPrintLogVerbosity && (Verbosity & ELogVerbosity::VerbosityMask) != ELogVerbosity::Log;
+
+		TCbWriter<1024> Writer;
+		Writer.BeginObject();
+		Writer.AddDateTime(ANSITEXTVIEW("time"), FDateTime::UtcNow());
+		Writer.AddString(ANSITEXTVIEW("level"), GetLevel(Verbosity));
+		AddMessage(Writer, V, Verbosity, Category, bShowCategory, bShowVerbosity);
+		if (bShowCategory || bShowVerbosity)
+		{
+			AddFormat(Writer, V, bShowCategory, bShowVerbosity);
+
+			Writer.BeginObject(ANSITEXTVIEW("properties"));
+			if (bShowCategory)
 			{
-				line.AppendChar(TEXT('\n'));
+				Writer.BeginObject(ANSITEXTVIEW("_channel"));
+				Writer.AddString(ANSITEXTVIEW("$type"), ANSITEXTVIEW("Channel"));
+				Writer.AddString(ANSITEXTVIEW("$text"), WriteToUtf8String<64>(Category));
+				Writer.EndObject();
+			}
+			if (bShowVerbosity)
+			{
+				Writer.BeginObject(ANSITEXTVIEW("_severity"));
+				Writer.AddString(ANSITEXTVIEW("$type"), ANSITEXTVIEW("Severity"));
+				Writer.AddString(ANSITEXTVIEW("$text"), ToString(Verbosity));
+				Writer.EndObject();
+			}
+			Writer.EndObject();
+		}
+		Writer.EndObject();
 
-				WriteConsoleW(GetStdHandle(STD_OUTPUT_HANDLE), *line, line.Len(), NULL, NULL);
+		WriteAsJson(Writer);
+	}
 
-				return;
+	FORCENOINLINE static void AddMessage(FCbWriter& Writer, const TCHAR* V, ELogVerbosity::Type Verbosity, const FName& Category, bool bShowCategory, bool bShowVerbosity)
+	{
+		TUtf8StringBuilder<512> Message;
+		if (bShowCategory)
+		{
+			Message << Category << ANSITEXTVIEW(": ");
+		}
+		if (bShowVerbosity)
+		{
+			Message << ToString(Verbosity) << ANSITEXTVIEW(": ");
+		}
+		Message << V;
+		Writer.AddString(ANSITEXTVIEW("message"), Message);
+	}
+
+	FORCENOINLINE static void AddFormat(FCbWriter& Writer, const TCHAR* Message, bool bShowCategory, bool bShowVerbosity)
+	{
+		TUtf8StringBuilder<512> Format;
+		if (bShowCategory)
+		{
+			Format << ANSITEXTVIEW("{_channel}: ");
+		}
+		if (bShowVerbosity)
+		{
+			Format << ANSITEXTVIEW("{_severity}: ");
+		}
+
+		// Escape {} in Message
+		for (constexpr FAsciiSet Brackets("{}");;)
+		{
+			const TCHAR* End = FAsciiSet::FindFirstOrEnd(Message, Brackets);
+			Format.Append(Message, UE_PTRDIFF_TO_INT32(End - Message));
+			if (!*End)
+			{
+				break;
 			}
+			Format.AppendChar(*End);
+			Format.AppendChar(*End);
+			Message = End + 1;
+		}
 
-			// fall through to standard printf path
-#endif
+		Writer.AddString(ANSITEXTVIEW("format"), Format);
+	}
 
-#if PLATFORM_TCHAR_IS_UTF8CHAR
-			printf("%s\n", (const char*)*line);
-#elif PLATFORM_TCHAR_IS_CHAR16
-			printf("%s\n", (const char*)StringCast<UTF8CHAR>(*line).Get());
-#elif PLATFORM_USE_LS_SPEC_FOR_WIDECHAR
-			// printf prints wchar_t strings just fine with %ls, while mixing printf()/wprintf() is not recommended (see https://stackoverflow.com/questions/8681623/printf-and-wprintf-in-single-c-code)
-			printf("%ls\n", *line);
-#else
-			static_assert(std::is_same_v<TCHAR, WIDECHAR>, "Assuming wide chars here");
-			wprintf(TEXT("%s\n"), *line);
+	FORCENOINLINE void SerializeRecordAsJson(const UE::FLogRecord& Record)
+	{
+		const bool bShowCategory = GPrintLogCategory && !Record.GetCategory().IsNone();
+		const bool bShowVerbosity = GPrintLogVerbosity && (Record.GetVerbosity() & ELogVerbosity::VerbosityMask) != ELogVerbosity::Log;
+
+		TCbWriter<1024> Writer;
+		Writer.BeginObject();
+		Writer.AddDateTime(ANSITEXTVIEW("time"), Record.GetTime().GetUtcTime());
+		Writer.AddString(ANSITEXTVIEW("level"), GetLevel(Record.GetVerbosity()));
+		AddMessage(Writer, Record, bShowCategory, bShowVerbosity);
+		if (bShowCategory || bShowVerbosity || Record.GetFields())
+		{
+			AddFormat(Writer, Record, bShowCategory, bShowVerbosity);
+
+			Writer.BeginObject(ANSITEXTVIEW("properties"));
+			if (bShowCategory)
+			{
+				Writer.BeginObject(ANSITEXTVIEW("_channel"));
+				Writer.AddString(ANSITEXTVIEW("$type"), ANSITEXTVIEW("Channel"));
+				Writer.AddString(ANSITEXTVIEW("$text"), WriteToUtf8String<64>(Record.GetCategory()));
+				Writer.EndObject();
+			}
+			if (bShowVerbosity)
+			{
+				Writer.BeginObject(ANSITEXTVIEW("_severity"));
+				Writer.AddString(ANSITEXTVIEW("$type"), ANSITEXTVIEW("Severity"));
+				Writer.AddString(ANSITEXTVIEW("$text"), ToString(Record.GetVerbosity()));
+				Writer.EndObject();
+			}
+			for (const FCbField& Field : Record.GetFields())
+			{
+				Writer.AddField(Field.GetName(), Field);
+			}
+			Writer.EndObject();
+		}
+		Writer.EndObject();
+
+		WriteAsJson(Writer);
+	}
+
+	FORCENOINLINE static void AddMessage(FCbWriter& Writer, const UE::FLogRecord& Record, bool bShowCategory, bool bShowVerbosity)
+	{
+		TUtf8StringBuilder<512> Message;
+		if (bShowCategory)
+		{
+			Message << Record.GetCategory() << ANSITEXTVIEW(": ");
+		}
+		if (bShowVerbosity)
+		{
+			Message << ToString(Record.GetVerbosity()) << ANSITEXTVIEW(": ");
+		}
+		Record.FormatMessageTo(Message);
+		Writer.AddString(ANSITEXTVIEW("message"), Message);
+	}
+
+	FORCENOINLINE static void AddFormat(FCbWriter& Writer, const UE::FLogRecord& Record, bool bShowCategory, bool bShowVerbosity)
+	{
+		TUtf8StringBuilder<512> Format;
+		if (bShowCategory)
+		{
+			Format << ANSITEXTVIEW("{_channel}: ");
+		}
+		if (bShowVerbosity)
+		{
+			Format << ANSITEXTVIEW("{_severity}: ");
+		}
+		Format.Append(Record.GetFormat());
+		Writer.AddString(ANSITEXTVIEW("format"), Format);
+	}
+
+	void WriteAsJson(const FCbWriter& Writer)
+	{
+		TArray<uint8, TInlineAllocator64<512>> Buffer;
+		Buffer.AddUninitialized((int64)Writer.GetSaveSize());
+		FCbFieldView Object = Writer.Save(MakeMemoryView(Buffer));
+
+	#if PLATFORM_WINDOWS
+		if (bIsConsoleOutput)
+		{
+			return WriteLine<WIDECHAR>(Object);
+		}
+	#endif
+	#if PLATFORM_TCHAR_IS_UTF8CHAR || PLATFORM_TCHAR_IS_CHAR16
+		WriteLine<UTF8CHAR>(Object);
+	#else
+		WriteLine<WIDECHAR>(Object);
+	#endif
+	}
+
+	template <typename CharType>
+	FORCENOINLINE void WriteLine(const FCbFieldView& Field)
+	{
+		TStringBuilderWithBuffer<CharType, 512> Line;
+		CompactBinaryToCompactJson(Field, Line);
+		Line.AppendChar('\n');
+		WriteLine(Line);
+	}
+
+	void WriteLine(FUtf8StringBuilderBase& Line) const
+	{
+		printf("%s", (const char*)*Line);
+		fflush(stdout);
+	}
+
+#if PLATFORM_WINDOWS || !(PLATFORM_TCHAR_IS_UTF8CHAR || PLATFORM_TCHAR_IS_CHAR16)
+	void WriteLine(FWideStringBuilderBase& Line) const
+	{
+	#if PLATFORM_WINDOWS
+		if (bIsConsoleOutput)
+		{
+			WriteConsoleW(GetStdHandle(STD_OUTPUT_HANDLE), *Line, Line.Len(), nullptr, nullptr);
+			return;
+		}
+	#endif
+
+	#if PLATFORM_USE_LS_SPEC_FOR_WIDECHAR
+		// printf prints wchar_t strings just fine with %ls, while mixing printf()/wprintf() is not recommended (see https://stackoverflow.com/questions/8681623/printf-and-wprintf-in-single-c-code)
+		printf("%ls", *Line);
+	#else
+		wprintf(TEXT("%s"), *Line);
+	#endif
+		fflush(stdout);
+	}
 #endif
 
-			fflush(stdout);
+	static FAnsiStringView GetLevel(ELogVerbosity::Type Verbosity)
+	{
+		switch (Verbosity & ELogVerbosity::VerbosityMask)
+		{
+		case ELogVerbosity::Fatal:
+			return ANSITEXTVIEW("Critical");
+		case ELogVerbosity::Error:
+			return ANSITEXTVIEW("Error");
+		case ELogVerbosity::Warning:
+			return ANSITEXTVIEW("Warning");
+		case ELogVerbosity::Display:
+		case ELogVerbosity::Log:
+		default:
+			return ANSITEXTVIEW("Information");
+		case ELogVerbosity::Verbose:
+		case ELogVerbosity::VeryVerbose:
+			return ANSITEXTVIEW("Debug");
 		}
 	}
 
 private:
 	ELogVerbosity::Type AllowedLogVerbosity = ELogVerbosity::Display;
 	bool				bIsConsoleOutput = false;
+	bool				bIsJsonOutput = false;
 
 #if PLATFORM_WINDOWS
 	static bool IsStdoutAttachedToConsole()
