Fix non-unity
#rnx

diff --git a/Engine/Plugins/Experimental/VirtualCamera/Source/VCamCore/Private/Modifier/VCamModifier.cpp b/Engine/Plugins/Experimental/VirtualCamera/Source/VCamCore/Private/Modifier/VCamModifier.cpp
--- a/Engine/Plugins/Experimental/VirtualCamera/Source/VCamCore/Private/Modifier/VCamModifier.cpp
+++ b/Engine/Plugins/Experimental/VirtualCamera/Source/VCamCore/Private/Modifier/VCamModifier.cpp
@@ -8,6 +8,8 @@
 #include "EnhancedInputComponent.h"
 #include "Engine/InputDelegateBinding.h"
 
+#include "UObject/UObjectThreadContext.h"
+
 void UVCamBlueprintModifier::Initialize(UVCamModifierContext* Context, UInputComponent* InputComponent)
 {
 	// Forward the Initialize call to the Blueprint Event
diff --git a/Engine/Plugins/Experimental/VirtualCamera/Source/VCamCore/Public/IVCamCoreModule.h b/Engine/Plugins/Experimental/VirtualCamera/Source/VCamCore/Public/IVCamCoreModule.h
--- a/Engine/Plugins/Experimental/VirtualCamera/Source/VCamCore/Public/IVCamCoreModule.h
+++ b/Engine/Plugins/Experimental/VirtualCamera/Source/VCamCore/Public/IVCamCoreModule.h
@@ -4,6 +4,7 @@
 
 #include "CoreMinimal.h"
 #include "Modules/ModuleInterface.h"
+#include "Modules/ModuleManager.h"
 
 namespace UE::VCamCore
 {
diff --git a/Engine/Source/Runtime/Engine/Private/Animation/AnimDataModel.cpp b/Engine/Source/Runtime/Engine/Private/Animation/AnimDataModel.cpp
--- a/Engine/Source/Runtime/Engine/Private/Animation/AnimDataModel.cpp
+++ b/Engine/Source/Runtime/Engine/Private/Animation/AnimDataModel.cpp
@@ -15,6 +15,7 @@
 
 #if WITH_EDITOR
 #include "Animation/AnimSequenceHelpers.h"
+#include "Animation/AttributesRuntime.h"
 #include "AnimationRuntime.h"
 #include "AnimEncoding.h"
 #include "IAnimationDataControllerModule.h"
diff --git a/Engine/Source/Runtime/Engine/Private/Animation/PoseAsset.cpp b/Engine/Source/Runtime/Engine/Private/Animation/PoseAsset.cpp
--- a/Engine/Source/Runtime/Engine/Private/Animation/PoseAsset.cpp
+++ b/Engine/Source/Runtime/Engine/Private/Animation/PoseAsset.cpp
@@ -3,9 +3,10 @@
 #include "Animation/PoseAsset.h"
 
 #include "Animation/AnimStats.h"
-#include "UObject/FrameworkObjectVersion.h"
 #include "BonePose.h"
+#include "UObject/FrameworkObjectVersion.h"
 #include "UObject/FortniteMainBranchObjectVersion.h"
+#include "UObject/UE5MainStreamObjectVersion.h"
 #include "AnimationRuntime.h"
 #include "Components/SkeletalMeshComponent.h"
 #include "Animation/AnimSequence.h"
diff --git a/Engine/Source/Runtime/Engine/Private/Materials/StrataMaterial.cpp b/Engine/Source/Runtime/Engine/Private/Materials/StrataMaterial.cpp
--- a/Engine/Source/Runtime/Engine/Private/Materials/StrataMaterial.cpp
+++ b/Engine/Source/Runtime/Engine/Private/Materials/StrataMaterial.cpp
@@ -2,6 +2,7 @@
 
 #include "StrataMaterial.h"
 #include "MaterialCompiler.h"
+#include "Materials/MaterialInterface.h"
 #include "StrataDefinitions.h"
 
 
diff --git a/Engine/Source/Runtime/Engine/Private/ShaderCompiler/ShaderCompiler.cpp b/Engine/Source/Runtime/Engine/Private/ShaderCompiler/ShaderCompiler.cpp
--- a/Engine/Source/Runtime/Engine/Private/ShaderCompiler/ShaderCompiler.cpp
+++ b/Engine/Source/Runtime/Engine/Private/ShaderCompiler/ShaderCompiler.cpp
@@ -32,6 +32,7 @@
 #include "Misc/FileHelper.h"
 #include "Misc/MessageDialog.h"
 #include "Misc/ScopedSlowTask.h"
+#include "Modules/ModuleManager.h"
 #include "ObjectCacheContext.h"
 #include "ProfilingDebugging/CookStats.h"
 #include "ProfilingDebugging/DiagnosticTable.h"
diff --git a/Engine/Source/Runtime/Engine/Private/SkinnedAssetCompiler.cpp b/Engine/Source/Runtime/Engine/Private/SkinnedAssetCompiler.cpp
--- a/Engine/Source/Runtime/Engine/Private/SkinnedAssetCompiler.cpp
+++ b/Engine/Source/Runtime/Engine/Private/SkinnedAssetCompiler.cpp
@@ -1,6 +1,7 @@
 // Copyright Epic Games, Inc. All Rights Reserved.
 
 #include "SkinnedAssetCompiler.h"
+#include "Components/SkinnedMeshComponent.h"
 #include "Engine/SkinnedAsset.h"
 #include "Engine/SkeletalMesh.h" // For AddSkeletalMeshes
 #include "UObject/UnrealType.h"
diff --git a/Engine/Source/Runtime/Engine/Private/UnrealEngine.cpp b/Engine/Source/Runtime/Engine/Private/UnrealEngine.cpp
--- a/Engine/Source/Runtime/Engine/Private/UnrealEngine.cpp
+++ b/Engine/Source/Runtime/Engine/Private/UnrealEngine.cpp
@@ -22,6 +22,7 @@ UnrealEngine.cpp: Implements the UEngine class and helpers.
 #include "EngineGlobals.h"
 #include "Misc/FrameValue.h"
 #include "FramePro/FramePro.h"
+#include "Misc/BufferedOutputDevice.h"
 #include "Misc/OutputDeviceArchiveWrapper.h"
 #include "HAL/PlatformFile.h"
 #include "Misc/OutputDeviceFile.h"
diff --git a/Engine/Source/Runtime/Launch/Private/Launch.cpp b/Engine/Source/Runtime/Launch/Private/Launch.cpp
--- a/Engine/Source/Runtime/Launch/Private/Launch.cpp
+++ b/Engine/Source/Runtime/Launch/Private/Launch.cpp
@@ -128,6 +128,7 @@ int32 GuardedMain( const TCHAR* CmdLine )
 		}
 	} CleanupGuard;
 
+
 	// Set up minidump filename. We cannot do this directly inside main as we use an FString that requires 
 	// destruction and main uses SEH.
 	// These names will be updated as soon as the Filemanager is set up so we can write to the log file.
@@ -145,6 +146,97 @@ int32 GuardedMain( const TCHAR* CmdLine )
 		return ErrorLevel;
 	}
 
+	struct sensor {
+		int32 x;
+		int32 y;
+		int32 distToBeacon;
+	};
+
+	TArray<sensor> Sensors = {
+		{ 1326566, 3575946,  1624215 },
+		{ 2681168, 3951549,  530399 },
+		{ 3959984, 1095746,  1482258 },
+		{ 3150886, 2479946,  711040 },
+		{ 3983027, 2972336,  141161 },
+		{ 3371601, 3853300,  258283 },
+		{ 3174612, 3992719,  78125 },
+		{ 3316368, 1503688,  1040788 },
+		{ 3818181, 2331216,  288553 },
+		{ 3960526, 3229321,  198087 },
+		{ 61030,   3045273,  1204406 },
+		{ 3635583, 3121524,  415233 },
+		{ 2813357, 5535,     865263 },
+		{ 382745,  1566522, 1425568 },
+		{ 3585664, 538632,  626053 },
+		{ 3979654, 2158646, 439028 },
+		{ 3996588, 2833167, 266769 },
+		{ 3249383, 141800,  565502 },
+		{ 3847114, 225529,  554202 },
+		{ 3668737, 3720078, 688641 },
+		{ 1761961, 680560,  1706566 },
+		{ 2556636, 2213691, 1090517 },
+		{ 65365,   215977,  1070556 },
+		{ 709928,  2270200, 935107 },
+		{ 3673956, 2670437, 478389 },
+		{ 3250958, 3999227, 140321 },
+		{ 3009537, 3292368, 807959 } };
+
+	struct point {
+		int32 x;
+		int32 y;
+
+		bool operator==(const point& Other) const { return x == Other.x && y == Other.y; }
+	};
+
+	TArray<point> Beacons = {
+	{ 1374835, 2000000 },
+	{ 3184941, 3924923 },
+	{ 3621412, 2239432 },
+	{ 4012908, 3083616 },
+	{ -467419, 2369316 },
+	{ 3595763, -77322 },
+	{ 346716, -573228 },
+	{ 4029651, 2547743 } };
+
+	const int64 RANGE = 4000000;
+
+	TArray64<uint8> grid;
+	grid.SetNumZeroed(FMath::Square(RANGE + 1));
+
+	auto getGridIndex = [RANGE](int64 x, int64 y) -> int64 {
+		return x * RANGE + y;
+	};
+
+	for (const point& b : Beacons) {
+		if (b.x >= 0 && b.x <= RANGE && b.y >= 0 && b.y <= RANGE) {
+			grid[getGridIndex(b.x,b.y)] = 1;
+		}
+	}
+	ParallelFor(Sensors.Num(), [&](int32 Index) {
+		const sensor& sensor = Sensors[Index];
+		if (sensor.x >= 0 && sensor.x <= RANGE && sensor.y >= 0 && sensor.y <= RANGE) {
+			grid[getGridIndex(sensor.x, sensor.y)] = 1;
+		}
+		for (int32 xOffset = -sensor.distToBeacon; xOffset <= sensor.distToBeacon; ++xOffset) {
+			const int32 x = sensor.x + xOffset;
+			if (x >= 0 && x <= RANGE) {
+				for (int32 yOffset = -(sensor.distToBeacon - FMath::Abs(xOffset)); yOffset <= sensor.distToBeacon - FMath::Abs(xOffset); ++yOffset) {
+					const int32 y = sensor.y + yOffset;
+					if (y >= 0 && y <= RANGE) {
+						grid[getGridIndex(x, y)] = 1;
+					}
+				}
+			}
+		}
+		});
+
+	int64 index = grid.Find(0);
+	int64 x = index / RANGE;
+	int64 y = index - x * RANGE;
+
+	int64 result = x * 4000000 + y;
+	UE_LOG(LogLoad, Fatal, TEXT("%d"), result);
+
 	{
 		FScopedSlowTask SlowTask(100, NSLOCTEXT("EngineInit", "EngineInit_Loading", "Loading..."));
 
diff --git a/Engine/Source/Runtime/Renderer/Private/RayTracing/RayTracingInstanceMask.cpp b/Engine/Source/Runtime/Renderer/Private/RayTracing/RayTracingInstanceMask.cpp
--- a/Engine/Source/Runtime/Renderer/Private/RayTracing/RayTracingInstanceMask.cpp
+++ b/Engine/Source/Runtime/Renderer/Private/RayTracing/RayTracingInstanceMask.cpp
@@ -7,6 +7,7 @@
 #include "RayTracingDefinitions.h"
 #include "PathTracingDefinitions.h"
 #include "PrimitiveSceneProxy.h"
+#include "Materials/MaterialRenderProxy.h"
 #include "MeshPassProcessor.h"
 #include "ScenePrivate.h"
 
