Adds GPU usages statistics to CSV profiles

diff --git a/Engine/Source/Runtime/Launch/Private/LaunchEngineLoop.cpp b/Engine/Source/Runtime/Launch/Private/LaunchEngineLoop.cpp
--- a/Engine/Source/Runtime/Launch/Private/LaunchEngineLoop.cpp
+++ b/Engine/Source/Runtime/Launch/Private/LaunchEngineLoop.cpp
@@ -1692,6 +1692,10 @@ static void UpdateCoreCsvStats_BeginFrame()
 	}
 }
 
+#if !UE_BUILD_SHIPPING
+CSV_DEFINE_CATEGORY(GPUUsage, true);
+#endif
+
 static void UpdateCoreCsvStats_EndFrame()
 {
 	if (!IsRunningDedicatedServer())
@@ -1724,6 +1728,14 @@ static void UpdateCoreCsvStats_EndFrame()
 		    TargetFPS = MaxFPSCVar->GetFloat();
 	    }
 	    CSV_CUSTOM_STAT_GLOBAL(MaxFrameTime, 1000.0f / TargetFPS, ECsvCustomStatOp::Set);
+
+		if (GRHISupportsGPUUsage && GNumExplicitGPUsForRendering == 1)
+		{
+			FRHIGPUUsageFractions GPUUsage = RHIGetGPUUsage(/* GPUIndex = */ 0);
+			CSV_CUSTOM_STAT(GPUUsage, Clock, GPUUsage.ClockScaling * 100.0f, ECsvCustomStatOp::Set);
+			CSV_CUSTOM_STAT(GPUUsage, Usage, GPUUsage.CurrentProcess * 100.0f, ECsvCustomStatOp::Set);
+			CSV_CUSTOM_STAT(GPUUsage, External, GPUUsage.ExternalProcesses * 100.0f, ECsvCustomStatOp::Set);
+		}
 #endif
 	}
 }
