Create a conflict mechanism to avoid Version.h getting in to a bad state
#rnx

diff --git a/Engine/Source/Runtime/Launch/Resources/Version.h b/Engine/Source/Runtime/Launch/Resources/Version.h
--- a/Engine/Source/Runtime/Launch/Resources/Version.h
+++ b/Engine/Source/Runtime/Launch/Resources/Version.h
@@ -54,10 +54,17 @@
 
 // These numbers define the banner UE version, and are the most significant numbers when ordering two engine versions (that is, a 5.12.* version is always 
 // newer than a 5.11.* version, regardless of the changelist that it was built with)
+// When updating these, also update the static_assert below
 #define ENGINE_MAJOR_VERSION	5
 #define ENGINE_MINOR_VERSION	2
 #define ENGINE_PATCH_VERSION	0
 
+// If this static_assert fires then Version.h has been updated without updating this code.
+// This line exists to cause conflicts when merging Version.h between streams so if one stream updates
+// the minor version while another updates the patch, you do not get a silently combined version that
+// was unintended.
+static_assert(ENGINE_MAJOR_VERSION == 5 && ENGINE_MINOR_VERSION == 1 && ENGINE_PATCH_VERSION == 1); //-V501
+
 // Macros for encoding strings
 #define VERSION_TEXT(x) TEXT(x)
 #define VERSION_STRINGIFY_2(x) VERSION_TEXT(#x)
