Correct static assert for UE5 stream
#rnx

diff --git a/Engine/Source/Runtime/Launch/Resources/Version.h b/Engine/Source/Runtime/Launch/Resources/Version.h
--- a/Engine/Source/Runtime/Launch/Resources/Version.h
+++ b/Engine/Source/Runtime/Launch/Resources/Version.h
@@ -63,7 +63,7 @@
 // This line exists to cause conflicts when merging Version.h between streams so if one stream updates
 // the minor version while another updates the patch, you do not get a silently combined version that
 // was unintended.
-static_assert(ENGINE_MAJOR_VERSION == 5 && ENGINE_MINOR_VERSION == 1 && ENGINE_PATCH_VERSION == 1); //-V501
+static_assert(ENGINE_MAJOR_VERSION == 5 && ENGINE_MINOR_VERSION == 2 && ENGINE_PATCH_VERSION == 0); //-V501
 
 // Macros for encoding strings
 #define VERSION_TEXT(x) TEXT(x)
