Allow external crash handling

While integrating a previous change, this line was left out, enabling external crash reporting tools to intercept exceptions.

diff --git a/Engine/Source/Runtime/Launch/Private/Windows/LaunchWindows.cpp b/Engine/Source/Runtime/Launch/Private/Windows/LaunchWindows.cpp
--- a/Engine/Source/Runtime/Launch/Private/Windows/LaunchWindows.cpp
+++ b/Engine/Source/Runtime/Launch/Private/Windows/LaunchWindows.cpp
@@ -245,7 +245,9 @@ LAUNCH_API int32 LaunchWindowsStartup( HINSTANCE hInInstance, HINSTANCE hPrevIns
 			GIsGuarded = 0;
 		}
 #if !PLATFORM_SEH_EXCEPTIONS_DISABLED
-		__except( GEnableInnerException ? EXCEPTION_EXECUTE_HANDLER : ReportCrash( GetExceptionInformation( ) ) )
+		__except( FPlatformMisc::GetCrashHandlingType() == ECrashHandlingType::Default
+				? ( GEnableInnerException ? EXCEPTION_EXECUTE_HANDLER : ReportCrash(GetExceptionInformation()) )
+				: EXCEPTION_CONTINUE_SEARCH )	
 		{
 			// Crashed.
 			ErrorLevel = 1;
