Move "ensure on commandlet error" processing into its own function so that crashreports have a more informative callstack. Also move it AFTER we print out the warning/error summary so that information is included in crash reports.

diff --git a/Engine/Source/Runtime/Launch/Private/LaunchEngineLoop.cpp b/Engine/Source/Runtime/Launch/Private/LaunchEngineLoop.cpp
--- a/Engine/Source/Runtime/Launch/Private/LaunchEngineLoop.cpp
+++ b/Engine/Source/Runtime/Launch/Private/LaunchEngineLoop.cpp
@@ -3344,6 +3344,16 @@ PRAGMA_ENABLE_DEPRECATION_WARNINGS
 	return 0;
 }
 
+void ConditionallyEnsureOnCommandletErrors(int32 InNumErrors)
+{
+	bool bEnsureOnError = false;
+	GConfig->GetBool(TEXT("Core.System"), TEXT("EnsureCommandletOnError"), bEnsureOnError, GEngineIni);
+	if (bEnsureOnError)
+	{
+		ensureMsgf(InNumErrors == 0, TEXT("Commandlet generated %s errors!"), InNumErrors);
+	}
+}
+
 int32 FEngineLoop::PreInitPostStartupScreen(const TCHAR* CmdLine)
 {
 	ON_SCOPE_EXIT{ GEnginePreInitPostStartupScreenEndTime = FPlatformTime::Seconds(); };
@@ -4031,13 +4041,6 @@ int32 FEngineLoop::PreInitPostStartupScreen(const TCHAR* CmdLine)
 				GWarn->GetErrors(AllErrors);
 				GWarn->GetWarnings(AllWarnings);
 
-				bool bEnsureOnError = false;
-				GConfig->GetBool(TEXT("Core.System"), TEXT("EnsureCommandletOnError"), bEnsureOnError, GEngineIni);
-				if (bEnsureOnError)
-				{
-					ensure(AllErrors.Num() == 0);
-				}
-
 				if (AllErrors.Num() || AllWarnings.Num())
 				{
 					SET_WARN_COLOR(COLOR_WHITE);
@@ -4093,6 +4096,9 @@ int32 FEngineLoop::PreInitPostStartupScreen(const TCHAR* CmdLine)
 					}
 				}
 
+				// Allow the caller to request that we issue an ensure when there are any errors from the executed commandlet
+				ConditionallyEnsureOnCommandletErrors(AllErrors.Num());
+
 				UE_LOG(LogInit, Display, TEXT(""));
 
 				if( ErrorLevel != 0 )
