Load WorldPartitionEditor module at editor boot to avoid trying to load it from ALT during object construction

#rnx

diff --git a/Engine/Source/Runtime/Launch/Private/LaunchEngineLoop.cpp b/Engine/Source/Runtime/Launch/Private/LaunchEngineLoop.cpp
--- a/Engine/Source/Runtime/Launch/Private/LaunchEngineLoop.cpp
+++ b/Engine/Source/Runtime/Launch/Private/LaunchEngineLoop.cpp
@@ -4562,6 +4562,10 @@ bool FEngineLoop::LoadStartupCoreModules()
 	// Required during construction of UAnimBlueprint which could happen from async loading thread.
 	// See UAnimBlueprint::UAnimBlueprint().
 	FModuleManager::Get().LoadModule(TEXT("AnimGraph"));
+
+	// Required during construction of UWorldPartitionEditorHash which could happen from async loading thread.
+	// See UWorldPartitionEditorHash::UWorldPartitionEditorHash().
+	FModuleManager::Get().LoadModule(TEXT("WorldPartitionEditor"));
 #endif
 
 	FModuleManager::Get().LoadModule(TEXT("PacketHandler"));
