Added a way to provide command line to Android apps vis adb shell

diff --git a/Engine/Build/Android/Java/src/com/epicgames/unreal/GameActivity.java.template b/Engine/Build/Android/Java/src/com/epicgames/unreal/GameActivity.java.template
--- a/Engine/Build/Android/Java/src/com/epicgames/unreal/GameActivity.java.template
+++ b/Engine/Build/Android/Java/src/com/epicgames/unreal/GameActivity.java.template
@@ -694,8 +694,15 @@ public class GameActivity extends $${gameActivitySuperClass}$$ implements Surfac
 		if(!TextUtils.isEmpty(tempCommandLine)) {
 			Log.debug("Using alternative commandline source");
 			// remember full commandline for regex checks by configrules
-			CommandLineFull = tempCommandLine;
-			parseCommandLineParams(tempCommandLine);
+			CommandLineFull += tempCommandLine;
+			parseCommandLineParams(CommandLineFull);
+			return;
+		}
+
+		if (CommandLineFull != "")
+		{
+			//if command line was provided via adb shell then skip reading cmd line from file
+			parseCommandLineParams(CommandLineFull);
 			return;
 		}
 
@@ -1427,7 +1434,7 @@ public class GameActivity extends $${gameActivitySuperClass}$$ implements Surfac
 			else if (SourceType.equals("SRC_CommandLine"))
 			{
 				bSourceIsCommandLine = true;
-				Source = CommandLineFull;;
+				Source = CommandLineFull;
 			}
 			else if (variables.containsKey(SourceType))
 			{
@@ -2341,6 +2348,8 @@ public class GameActivity extends $${gameActivitySuperClass}$$ implements Surfac
 			getWindow().getDecorView().setImportantForAutofill(View.IMPORTANT_FOR_AUTOFILL_NO_EXCLUDE_DESCENDANTS);
 		}
 
+		boolean bIsShipping = nativeIsShippingBuild();
+
 		// create splashscreen dialog (if launched by SplashActivity)
 		_extrasBundle = getIntent().getExtras();
 		if (_extrasBundle != null)
@@ -2391,6 +2400,12 @@ public class GameActivity extends $${gameActivitySuperClass}$$ implements Surfac
 				e.printStackTrace();
 			}
 
+			String CmdLine = _extrasBundle.getString("cmdline");
+			if (CmdLine != null && !bIsShipping)
+			{
+				CommandLineFull = CmdLine;
+			}
+
 			if (UseDisplayCutout)
 			{
 				// will not be true if not Android Pie or later
@@ -2414,7 +2429,6 @@ public class GameActivity extends $${gameActivitySuperClass}$$ implements Surfac
 		}
 
 		// Check for target sdk.  If 23 or higher then warn that permission handling may mean features don't work if user denies them.
-		boolean bIsShipping = nativeIsShippingBuild();
 		if (ANDROID_BUILD_VERSION >= 23 && targetSdkVersion >= 23) //23 is the API level (Marshmallow) where runtime permission handling is available
 		{
 			Log.debug("Target SDK is " + targetSdkVersion + ".  This may cause issues if permissions are denied by the user." );	
diff --git a/Engine/Source/Runtime/Launch/Private/Android/LaunchAndroid.cpp b/Engine/Source/Runtime/Launch/Private/Android/LaunchAndroid.cpp
--- a/Engine/Source/Runtime/Launch/Private/Android/LaunchAndroid.cpp
+++ b/Engine/Source/Runtime/Launch/Private/Android/LaunchAndroid.cpp
@@ -164,6 +164,7 @@ extern bool AndroidThunkCpp_VirtualInputIgnoreClick(int, int);
 extern bool AndroidThunkCpp_IsVirtuaKeyboardShown();
 extern bool AndroidThunkCpp_IsWebViewShown();
 extern void AndroidThunkCpp_RestartApplication(const FString& IntentString);
+extern FString AndroidThunkCpp_GetIntentExtrasString(const FString& Key);
 
 // Base path for file accesses
 extern FString GFilePathBase;
@@ -280,60 +281,73 @@ static void InitCommandLine()
 	// initialize the command line to an empty string
 	FCommandLine::Set(TEXT(""));
 
-	AAssetManager* AssetMgr = AndroidThunkCpp_GetAssetManager();
-	AAsset* asset = AAssetManager_open(AssetMgr, TCHAR_TO_UTF8(TEXT("UECommandLine.txt")), AASSET_MODE_BUFFER);
-	if (nullptr != asset)
+#if !UE_BUILD_SHIPPING
+	FString CmdLine = AndroidThunkCpp_GetIntentExtrasString(TEXT("cmdline"));
+	if (!CmdLine.IsEmpty())
 	{
-		const void* FileContents = AAsset_getBuffer(asset);
-		int32 FileLength = AAsset_getLength(asset);
+		CmdLine.TrimEndInline();
+		FCommandLine::Append(*CmdLine);
 
-		char CommandLine[CMD_LINE_MAX];
-		FileLength = (FileLength < CMD_LINE_MAX - 1) ? FileLength : CMD_LINE_MAX - 1;
-		memcpy(CommandLine, FileContents, FileLength);
-		CommandLine[FileLength] = '\0';
+		FPlatformMisc::LowLevelOutputDebugStringf(TEXT("adb am start command line override: %s"), FCommandLine::Get());
+	}
+	else
+#endif
+	{
+		AAssetManager* AssetMgr = AndroidThunkCpp_GetAssetManager();
+		AAsset* asset = AAssetManager_open(AssetMgr, TCHAR_TO_UTF8(TEXT("UECommandLine.txt")), AASSET_MODE_BUFFER);
+		if (nullptr != asset)
+		{
+			const void* FileContents = AAsset_getBuffer(asset);
+			int32 FileLength = AAsset_getLength(asset);
+
+			char CommandLine[CMD_LINE_MAX];
+			FileLength = (FileLength < CMD_LINE_MAX - 1) ? FileLength : CMD_LINE_MAX - 1;
+			memcpy(CommandLine, FileContents, FileLength);
+			CommandLine[FileLength] = '\0';
+
+			AAsset_close(asset);
+
+			// chop off trailing spaces
+			while (*CommandLine && isspace(CommandLine[strlen(CommandLine) - 1]))
+			{
+				CommandLine[strlen(CommandLine) - 1] = 0;
+			}
 
-		AAsset_close(asset);
+			FCommandLine::Append(UTF8_TO_TCHAR(CommandLine));
+			FPlatformMisc::LowLevelOutputDebugStringf(TEXT("APK Commandline: %s"), FCommandLine::Get());
+		}
 
-		// chop off trailing spaces
-		while (*CommandLine && isspace(CommandLine[strlen(CommandLine) - 1]))
+		// read in the command line text file from the sdcard if it exists
+		FString CommandLineFilePath = GFilePathBase + FString("/UnrealGame/") + (!FApp::IsProjectNameEmpty() ? FApp::GetProjectName() : FPlatformProcess::ExecutableName()) + FString("/UECommandLine.txt");
+		FILE* CommandLineFile = fopen(TCHAR_TO_UTF8(*CommandLineFilePath), "r");
+		if (CommandLineFile == NULL)
 		{
-			CommandLine[strlen(CommandLine) - 1] = 0;
+			// if that failed, try the lowercase version
+			CommandLineFilePath = CommandLineFilePath.Replace(TEXT("UECommandLine.txt"), TEXT("uecommandline.txt"));
+			CommandLineFile = fopen(TCHAR_TO_UTF8(*CommandLineFilePath), "r");
 		}
 
-		FCommandLine::Append(UTF8_TO_TCHAR(CommandLine));
-		FPlatformMisc::LowLevelOutputDebugStringf(TEXT("APK Commandline: %s"), FCommandLine::Get());
-	}
+		if (CommandLineFile)
+		{
+			FPlatformMisc::LowLevelOutputDebugStringf(TEXT("Using override commandline file: %s"), *CommandLineFilePath);
 
-	// read in the command line text file from the sdcard if it exists
-	FString CommandLineFilePath = GFilePathBase + FString("/UnrealGame/") + (!FApp::IsProjectNameEmpty() ? FApp::GetProjectName() : FPlatformProcess::ExecutableName()) + FString("/UECommandLine.txt");
-	FILE* CommandLineFile = fopen(TCHAR_TO_UTF8(*CommandLineFilePath), "r");
-	if(CommandLineFile == NULL)
-	{
-		// if that failed, try the lowercase version
-		CommandLineFilePath = CommandLineFilePath.Replace(TEXT("UECommandLine.txt"), TEXT("uecommandline.txt"));
-		CommandLineFile = fopen(TCHAR_TO_UTF8(*CommandLineFilePath), "r");
-	}
+			char CommandLine[CMD_LINE_MAX];
+			fgets(CommandLine, UE_ARRAY_COUNT(CommandLine) - 1, CommandLineFile);
 
-	if(CommandLineFile)
-	{
-		FPlatformMisc::LowLevelOutputDebugStringf(TEXT("Using override commandline file: %s"), *CommandLineFilePath);
+			fclose(CommandLineFile);
 
-		char CommandLine[CMD_LINE_MAX];
-		fgets(CommandLine, UE_ARRAY_COUNT(CommandLine) - 1, CommandLineFile);
+			// chop off trailing spaces
+			while (*CommandLine && isspace(CommandLine[strlen(CommandLine) - 1]))
+			{
+				CommandLine[strlen(CommandLine) - 1] = 0;
+			}
 
-		fclose(CommandLineFile);
+			// initialize the command line to an empty string
+			FCommandLine::Set(TEXT(""));
 
-		// chop off trailing spaces
-		while (*CommandLine && isspace(CommandLine[strlen(CommandLine) - 1]))
-		{
-			CommandLine[strlen(CommandLine) - 1] = 0;
+			FCommandLine::Append(UTF8_TO_TCHAR(CommandLine));
+			FPlatformMisc::LowLevelOutputDebugStringf(TEXT("Override Commandline: %s"), FCommandLine::Get());
 		}
-
-		// initialize the command line to an empty string
-		FCommandLine::Set(TEXT(""));
-
-		FCommandLine::Append(UTF8_TO_TCHAR(CommandLine));
-		FPlatformMisc::LowLevelOutputDebugStringf(TEXT("Override Commandline: %s"), FCommandLine::Get());
 	}
 
 #if !UE_BUILD_SHIPPING
