Add EnginePreInitMemory LLM tag to  EnginePreInitPostStartupScreen

Fixes several untagged allocations which occur during this phase of engine startup.

diff --git a/Engine/Source/Runtime/Launch/Private/LaunchEngineLoop.cpp b/Engine/Source/Runtime/Launch/Private/LaunchEngineLoop.cpp
--- a/Engine/Source/Runtime/Launch/Private/LaunchEngineLoop.cpp
+++ b/Engine/Source/Runtime/Launch/Private/LaunchEngineLoop.cpp
@@ -3362,6 +3362,7 @@ int32 FEngineLoop::PreInitPostStartupScreen(const TCHAR* CmdLine)
 {
 	ON_SCOPE_EXIT{ GEnginePreInitPostStartupScreenEndTime = FPlatformTime::Seconds(); };
 	SCOPED_BOOT_TIMING("FEngineLoop::PreInitPostStartupScreen");
+	LLM_SCOPE(ELLMTag::EnginePreInitMemory);
 
 	if (IsEngineExitRequested())
 	{
