Add -fastexit command-line argument to terminate process as soon as a commandlet has finished
  - Saves 3.3s at the end of a small warm cook

#rnx

diff --git a/Engine/Source/Runtime/Launch/Private/LaunchEngineLoop.cpp b/Engine/Source/Runtime/Launch/Private/LaunchEngineLoop.cpp
--- a/Engine/Source/Runtime/Launch/Private/LaunchEngineLoop.cpp
+++ b/Engine/Source/Runtime/Launch/Private/LaunchEngineLoop.cpp
@@ -4154,6 +4154,11 @@ int32 FEngineLoop::PreInitPostStartupScreen(const TCHAR* CmdLine)
 				}
 			}
 
+			if (FParse::Param(FCommandLine::Get(), TEXT("fastexit")))
+			{
+				FPlatformMisc::RequestExitWithStatus(true, ErrorLevel);
+			}
+
 			// We're ready to exit!
 			return ErrorLevel;
 		}
