- Fixed quotes being removed from FCommandLine parameters when breaking it up, removing entry 0, then restoring the commandline
- Moved a static function from LaunchEngineLoop.h that was never used externally to a static in LaunchEngineLoop.cpp

diff --git a/Engine/Source/Runtime/Launch/Private/LaunchEngineLoop.cpp b/Engine/Source/Runtime/Launch/Private/LaunchEngineLoop.cpp
--- a/Engine/Source/Runtime/Launch/Private/LaunchEngineLoop.cpp
+++ b/Engine/Source/Runtime/Launch/Private/LaunchEngineLoop.cpp
@@ -1784,6 +1784,102 @@
 	GInputTime = FPlatformTime::Cycles64();
 }
 
+static TArray<FString> TokenizeCommandline(const TCHAR* CmdLine, bool bRetainQuotes)
+{
+	TArray<FString> TokenArray;
+
+	const TCHAR* ParsedCmdLine = CmdLine;
+
+	while (*ParsedCmdLine)
+	{
+		FString Token;
+				
+		// skip over whitespace to look for a quote
+		while (FChar::IsWhitespace(*ParsedCmdLine))
+		{
+			ParsedCmdLine++;
+		}
+
+		// if we want to keep quotes around the token, and the first character is a quote, then FToken::Parse
+		// will remove the quotes, so put them back
+		if (bRetainQuotes && (*ParsedCmdLine == TEXT('"')))
+		{
+			FParse::Token(ParsedCmdLine, Token, 0);
+			Token = FString::Printf(TEXT("\"%s\""), *Token);
+		}
+		else
+		{
+			FParse::Token(ParsedCmdLine, Token, 0);
+			Token.TrimStartAndEndInline();
+		}
+
+		if (!Token.IsEmpty())
+		{
+			TokenArray.Add(MoveTemp(Token));
+		}
+	}
+
+	return MoveTemp(TokenArray);
+}
+
+/** Enumeration representing the type of the command-line argument representing the game (typically the first argument). */
+enum class EGameStringType
+{
+	GameName,
+	ProjectPath,
+	ProjectShortName,
+	Unknown
+};
+
+/**
+  * Finds a command-line argument representing the game, removes it from the array and returns.
+  *
+ * @param TokenArray    Array of the tokenized command line
+ * @param QuotedTokenArray  Parellel array that maintains quotes around params
+  * @param OutStringType Type of the game string found.
+  * @return String representing the game command-line parameter; empty string if not found (OutStringType == Unknown in such case).
+  */
+static FString ExtractGameStringArgument(TArray<FString>& TokenArray, TArray<FString>& QuotedTokenArray, EGameStringType& OutStringType)
+{
+	for (int32 I = 0; I < TokenArray.Num(); ++I)
+	{
+		FString NormalizedToken = TokenArray[I];
+
+		// Path returned by FPaths::GetProjectFilePath() is normalized, so may have symlinks and ~ resolved and may differ from the original path to .uproject passed in the command line
+		FPaths::NormalizeFilename(NormalizedToken);
+
+		const bool bTokenIsGameName                 = (FApp::HasProjectName()         && TokenArray[I]   == FApp::GetProjectName());
+		const bool bTokenIsGameProjectFilePath      = (FPaths::IsProjectFilePathSet() && NormalizedToken == FPaths::GetProjectFilePath());
+		const bool bTokenIsGameProjectFileShortName = (FPaths::IsProjectFilePathSet() && TokenArray[I]   == FPaths::GetCleanFilename(FPaths::GetProjectFilePath()));
+
+		if (bTokenIsGameName || bTokenIsGameProjectFilePath || bTokenIsGameProjectFileShortName)
+		{
+			if (bTokenIsGameName)
+			{
+				OutStringType = EGameStringType::GameName;
+			}
+			else if (bTokenIsGameProjectFilePath)
+			{
+				OutStringType = EGameStringType::ProjectPath;
+			}
+			else if (bTokenIsGameProjectFileShortName)
+			{
+				OutStringType = EGameStringType::ProjectShortName;
+			}
+
+			FString Result = TokenArray[I];
+
+			TokenArray.RemoveAt(I);
+			QuotedTokenArray.RemoveAt(I);
+
+			return Result;
+		}
+	}
+
+	OutStringType = EGameStringType::Unknown;
+
+	return FString();
+}
 
 
 DECLARE_CYCLE_STAT(TEXT("FEngineLoop::PreInitPreStartupScreen.AfterStats"), STAT_FEngineLoop_PreInitPreStartupScreen_AfterStats, STATGROUP_LoadTime);
@@ -2178,10 +2274,13 @@
 
 	FString TokenToForward;
 
-	TArray<FString> TokenArray     = TokenizeCommandline(CmdLine);
+	// these tokens are use later to restore a single commandline string, so keep a version that has quotes around tokens,
+	// in case there are spaces in a token which will break parsing that single string
+	TArray<FString> TokenArray     		= TokenizeCommandline(CmdLine, false);
+	TArray<FString> QuotedTokenArray    = TokenizeCommandline(CmdLine, true);
 
 	EGameStringType GameStringType = EGameStringType::Unknown;
-	FString         GameString     = ExtractGameStringArgument(TokenArray, GameStringType);
+	FString         GameString     = ExtractGameStringArgument(TokenArray, QuotedTokenArray, GameStringType);
 
 	auto IsModeSelected = [&bHasCommandletToken, &bHasEditorToken, &bIsRegularClient, &bIsRunningAsDedicatedServer]()
 	{
@@ -2189,7 +2288,7 @@
 	};
 #if UE_EDITOR || WITH_ENGINE || WITH_EDITOR
 	FString CommandletCommandLine;
-	auto SetIsRunningAsCommandlet = [&CommandletCommandLine, &TokenArray,
+	auto SetIsRunningAsCommandlet = [&CommandletCommandLine, &QuotedTokenArray,
 		&bHasCommandletToken, &bIsRunningAsDedicatedServer, &bHasEditorToken, &bIsRegularClient, &IsModeSelected, &TokenToForward]
 		(FStringView CommandletName)
 	{
@@ -2211,7 +2310,7 @@
 		}
 #endif
 
-		CommandletCommandLine = FString::Join(TokenArray, TEXT(" "));
+		CommandletCommandLine = FString::Join(QuotedTokenArray, TEXT(" "));
 #if WITH_EDITOR
 		if (CommandletName == TEXTVIEW("cookcommandlet"))
 		{
@@ -2330,7 +2429,7 @@
 		if (GameStringType != EGameStringType::Unknown)
 		{
 			// Set a new command-line that doesn't include the game name as the first argument.
-			FCommandLine::Set(*FString::Join(TokenArray, TEXT(" ")));
+			FCommandLine::Set(*FString::Join(QuotedTokenArray, TEXT(" ")));
 
 			// Remove spurious project file tokens (which can happen on some platforms that combine commandlines).
 			// This handles extra .uprojects, but if you run with MyGame MyGame, we can't tell if the second MyGame is a map or not.
@@ -2367,8 +2466,10 @@
 				FString LocalToken = TokenArray[0];
 				TokenArray.Add(LocalToken);
 				TokenArray.RemoveAt(0);
+				QuotedTokenArray.Add(LocalToken);
+				QuotedTokenArray.RemoveAt(0);
 
-				FCommandLine::Set(*FString::Join(TokenArray, TEXT(" ")));
+				FCommandLine::Set(*FString::Join(QuotedTokenArray, TEXT(" ")));
 			}
 		}
 #endif
@@ -4368,69 +4469,7 @@
 #endif // WITH_EDITOR
 }
 
-TArray<FString> FEngineLoop::TokenizeCommandline(const TCHAR* CmdLine)
-{
-	TArray<FString> TokenArray;
 
-	const TCHAR* ParsedCmdLine = CmdLine;
 
-	while (*ParsedCmdLine)
-	{
-		FString Token;
-		
-		FParse::Token(ParsedCmdLine, Token, 0);
-
-		Token.TrimStartAndEndInline();
-
-		if (!Token.IsEmpty())
-		{
-			TokenArray.Add(MoveTemp(Token));
-		}
-	}
-
-	return MoveTemp(TokenArray);
-}
-
-FString FEngineLoop::ExtractGameStringArgument(TArray<FString>& TokenArray, EGameStringType& OutStringType)
-{
-	for (int32 I = 0; I < TokenArray.Num(); ++I)
-	{
-		FString NormalizedToken = TokenArray[I];
-
-		// Path returned by FPaths::GetProjectFilePath() is normalized, so may have symlinks and ~ resolved and may differ from the original path to .uproject passed in the command line
-		FPaths::NormalizeFilename(NormalizedToken);
-
-		const bool bTokenIsGameName                 = (FApp::HasProjectName()         && TokenArray[I]   == FApp::GetProjectName());
-		const bool bTokenIsGameProjectFilePath      = (FPaths::IsProjectFilePathSet() && NormalizedToken == FPaths::GetProjectFilePath());
-		const bool bTokenIsGameProjectFileShortName = (FPaths::IsProjectFilePathSet() && TokenArray[I]   == FPaths::GetCleanFilename(FPaths::GetProjectFilePath()));
-
-		if (bTokenIsGameName || bTokenIsGameProjectFilePath || bTokenIsGameProjectFileShortName)
-		{
-			if (bTokenIsGameName)
-			{
-				OutStringType = EGameStringType::GameName;
-			}
-			else if (bTokenIsGameProjectFilePath)
-			{
-				OutStringType = EGameStringType::ProjectPath;
-			}
-			else if (bTokenIsGameProjectFileShortName)
-			{
-				OutStringType = EGameStringType::ProjectShortName;
-			}
-
-			FString Result = TokenArray[I];
-
-			TokenArray.RemoveAt(I);
-
-			return Result;
-		}
-	}
-
-	OutStringType = EGameStringType::Unknown;
-
-	return FString();
-}
-
 #if WITH_ENGINE
 
diff --git a/Engine/Source/Runtime/Launch/Public/LaunchEngineLoop.h b/Engine/Source/Runtime/Launch/Public/LaunchEngineLoop.h
--- a/Engine/Source/Runtime/Launch/Public/LaunchEngineLoop.h
+++ b/Engine/Source/Runtime/Launch/Public/LaunchEngineLoop.h
@@ -178,26 +178,6 @@ protected:
 
 private:
 
-	/** Enumeration representing the type of the command-line argument representing the game (typically the first argument). */
-	enum class EGameStringType
-	{
-		GameName,
-		ProjectPath,
-		ProjectShortName,
-		Unknown
-	};
-
-	/** Takes a command-line string and returns an array of tokens (splits arguments using whitespaces). */
-	static TArray<FString> TokenizeCommandline(const TCHAR* CmdLine);
-
-	/**
-	  * Finds a command-line argument representing the game, removes it from the array and returns.
-	  * 
-	  * @param TokenArray    Array of the tokenized command line
-	  * @param OutStringType Type of the game string found.
-	  * @return String representing the game command-line parameter; empty string if not found (OutStringType == Unknown in such case).
-	  */
-	static FString ExtractGameStringArgument(TArray<FString>& TokenArray, EGameStringType& OutStringType);
 
 #if WITH_ENGINE
 
