Fix format string in ensure on commandlet error

diff --git a/Engine/Source/Runtime/Launch/Private/LaunchEngineLoop.cpp b/Engine/Source/Runtime/Launch/Private/LaunchEngineLoop.cpp
--- a/Engine/Source/Runtime/Launch/Private/LaunchEngineLoop.cpp
+++ b/Engine/Source/Runtime/Launch/Private/LaunchEngineLoop.cpp
@@ -3443,7 +3443,7 @@ void ConditionallyEnsureOnCommandletErrors(int32 InNumErrors)
 	GConfig->GetBool(TEXT("Core.System"), TEXT("EnsureCommandletOnError"), bEnsureOnError, GEngineIni);
 	if (bEnsureOnError)
 	{
-		ensureMsgf(InNumErrors == 0, TEXT("Commandlet generated %s errors!"), InNumErrors);
+		ensureMsgf(InNumErrors == 0, TEXT("Commandlet generated %d errors!"), InNumErrors);
 	}
 }
 
