Prevent deadlock while backgrounding and native keyboard events occur.

diff --git a/Engine/Source/Runtime/Launch/Private/Android/AndroidJNI.cpp b/Engine/Source/Runtime/Launch/Private/Android/AndroidJNI.cpp
--- a/Engine/Source/Runtime/Launch/Private/Android/AndroidJNI.cpp
+++ b/Engine/Source/Runtime/Launch/Private/Android/AndroidJNI.cpp
@@ -1023,12 +1023,12 @@ JNI_METHOD void Java_com_epicgames_unreal_GameActivity_nativeVirtualKeyboardResu
 	{
 		if (VirtualKeyboardWidget.IsValid())
 		{
-			auto Contents = FJavaHelper::FStringFromParam(jenv, contents);
-			
+			FString Contents = FJavaHelper::FStringFromParam(jenv, contents);
+
 			// call to set the widget text on game thread
 			if (FTaskGraphInterface::IsRunning())
 			{
-				FGraphEventRef SetWidgetText = FFunctionGraphTask::CreateAndDispatchWhenReady([&]()
+				FFunctionGraphTask::CreateAndDispatchWhenReady([Contents=Contents]()
 				{
 					TSharedPtr<IVirtualKeyboardEntry> LockedKeyboardWidget(VirtualKeyboardWidget.Pin());
 					if (LockedKeyboardWidget.IsValid())
@@ -1039,7 +1039,6 @@ JNI_METHOD void Java_com_epicgames_unreal_GameActivity_nativeVirtualKeyboardResu
 					// release reference
 					VirtualKeyboardWidget.Reset();
 				}, TStatId(), NULL, ENamedThreads::GameThread);
-				FTaskGraphInterface::Get().WaitUntilTaskCompletes(SetWidgetText);
 			}
 			else
 			{
@@ -1060,12 +1059,12 @@ JNI_METHOD void Java_com_epicgames_unreal_GameActivity_nativeVirtualKeyboardChan
 {
 	if (VirtualKeyboardWidget.IsValid())
 	{
-		auto Contents = FJavaHelper::FStringFromParam(jenv, contents);
+		FString Contents = FJavaHelper::FStringFromParam(jenv, contents);
 		
 		// call to set the widget text on game thread
 		if (FTaskGraphInterface::IsRunning())
 		{
-			FGraphEventRef SetWidgetText = FFunctionGraphTask::CreateAndDispatchWhenReady([&]()
+			FFunctionGraphTask::CreateAndDispatchWhenReady([Contents=Contents]()
 			{
 				TSharedPtr<IVirtualKeyboardEntry> LockedKeyboardWidget(VirtualKeyboardWidget.Pin());
 				if (LockedKeyboardWidget.IsValid())
@@ -1073,7 +1072,6 @@ JNI_METHOD void Java_com_epicgames_unreal_GameActivity_nativeVirtualKeyboardChan
 					LockedKeyboardWidget->SetTextFromVirtualKeyboard(FText::FromString(Contents), ETextEntryType::TextEntryUpdated);
 				}
 			}, TStatId(), NULL, ENamedThreads::GameThread);
-			FTaskGraphInterface::Get().WaitUntilTaskCompletes(SetWidgetText);
 		}
 	}
 }
@@ -1839,8 +1837,6 @@ JNI_METHOD void Java_com_epicgames_unreal_GameActivity_nativeOnSafetyNetAttestat
 	}
 }
 
-
-
 #include "Async/TaskGraphInterfaces.h"
 
 static jmethodID ReplyMethod;
