Add UEngine::IsRenderingSuspended(), which can be used to reduce GPU performance impact on other applications

* When this method returns "true", all viewport and Slate widget rendering is skipped
* This is especially effective with D3D12.EvictAllResidentResourcesInBackground=1, which immediately releases most of the VRAM when the engine does not touch any GPU resources

diff --git a/Engine/Source/Runtime/Engine/Classes/Engine/Engine.h b/Engine/Source/Runtime/Engine/Classes/Engine/Engine.h
--- a/Engine/Source/Runtime/Engine/Classes/Engine/Engine.h
+++ b/Engine/Source/Runtime/Engine/Classes/Engine/Engine.h
@@ -2346,6 +2346,9 @@ public:
 	/** Whether we're allowed to do frame rate smoothing */
 	virtual bool IsAllowedFramerateSmoothing() const;
 
+	/** Whether the application should avoid rendering anything to give GPU resources to other applications */
+	virtual bool IsRenderingSuspended() const { return false; }
+
 	/** Update FApp::Timecode. */
 	void UpdateTimecode();
 
diff --git a/Engine/Source/Runtime/Engine/Private/GameEngine.cpp b/Engine/Source/Runtime/Engine/Private/GameEngine.cpp
--- a/Engine/Source/Runtime/Engine/Private/GameEngine.cpp
+++ b/Engine/Source/Runtime/Engine/Private/GameEngine.cpp
@@ -1878,7 +1878,7 @@ void UGameEngine::Tick( float DeltaSeconds, bool bIdleMode )
 		}
 	}
 
-	if (!bIdleMode && !IsRunningDedicatedServer() && !IsRunningCommandlet() && FEmbeddedCommunication::IsAwakeForRendering())
+	if (!bIdleMode && !IsRenderingSuspended() && !IsRunningDedicatedServer() && !IsRunningCommandlet() && FEmbeddedCommunication::IsAwakeForRendering())
 	{
 		// Render everything.
 		RedrawViewports();
diff --git a/Engine/Source/Runtime/Launch/Private/LaunchEngineLoop.cpp b/Engine/Source/Runtime/Launch/Private/LaunchEngineLoop.cpp
--- a/Engine/Source/Runtime/Launch/Private/LaunchEngineLoop.cpp
+++ b/Engine/Source/Runtime/Launch/Private/LaunchEngineLoop.cpp
@@ -5931,8 +5931,10 @@ void FEngineLoop::Tick()
 		// We split separate this action from the one above to permit running network replication concurrent with slate widget ticking and painting.
 		if (FSlateApplication::IsInitialized() && !bIdleMode)
 		{
+			const bool bRenderingSuspended = GEngine->IsRenderingSuspended();
+
 			FMoviePlayerProxy::SetIsSlateThreadAllowed(false);
-			FSlateApplication::Get().Tick(ESlateTickType::TimeAndWidgets);
+			FSlateApplication::Get().Tick(bRenderingSuspended ? ESlateTickType::Time : ESlateTickType::TimeAndWidgets);
 			FMoviePlayerProxy::SetIsSlateThreadAllowed(true);
 		}
 
