CrashHandling: instrumentation for missing callstacks attempt 4.
#rnx

diff --git a/Engine/Source/Runtime/Launch/Private/Windows/LaunchWindows.cpp b/Engine/Source/Runtime/Launch/Private/Windows/LaunchWindows.cpp
--- a/Engine/Source/Runtime/Launch/Private/Windows/LaunchWindows.cpp
+++ b/Engine/Source/Runtime/Launch/Private/Windows/LaunchWindows.cpp
@@ -250,6 +250,7 @@ LAUNCH_API int32 LaunchWindowsStartup( HINSTANCE hInInstance, HINSTANCE hPrevIns
 				: EXCEPTION_CONTINUE_SEARCH )	
 		{
 			UE_LOG(LogLaunchWindows, Log, TEXT("LaunchWindowsStartup.ExceptionHandler"));
+			UE_LOG(LogLaunchWindows, Error, TEXT("LaunchWindowsStartup.ExceptionHandler"));
 			// Crashed.
 			ErrorLevel = 1;
 			if(GError)
