Android Local Notifications changes
Allow notifications to specify a pinned ID, By doing this it will only show the latest notification sent from that ID.
Currently no IOS implementation.

diff --git a/Engine/Build/Android/Java/src/com/epicgames/unreal/GameActivity.java.template b/Engine/Build/Android/Java/src/com/epicgames/unreal/GameActivity.java.template
--- a/Engine/Build/Android/Java/src/com/epicgames/unreal/GameActivity.java.template
+++ b/Engine/Build/Android/Java/src/com/epicgames/unreal/GameActivity.java.template
@@ -4980,6 +4980,36 @@ public class GameActivity extends $${gameActivitySuperClass}$$ implements Surfac
 		}
 	}
 
+	public static void AddLocalNotificationID(Context context, int id)
+	{
+		SharedPreferences preferences = context.getSharedPreferences("LocalNotificationPreferences", MODE_PRIVATE);
+		SharedPreferences.Editor editor = preferences.edit();
+		String notificationIDs = preferences.getString("notificationIDs", "");
+		if (notificationIDs.length() == 0)
+		{
+			editor.putString("notificationIDs", Integer.toString(id));
+			editor.commit();
+		}
+		else
+		{
+			String[] parts = notificationIDs.split("-");
+			ArrayList<Integer> iParts = new ArrayList<Integer>();
+			for (String part : parts)
+			{
+				if (part.length() > 0)
+				{
+					iParts.add(Integer.parseInt(part));
+				}
+			}
+			if (!iParts.contains(id))
+			{
+				editor.putString("notificationIDs", notificationIDs + "-" + id);
+				editor.commit();
+			}
+		}
+		
+	}
+
 	public static int LocalNotificationGetID(Context context)
 	{
 		SharedPreferences preferences = context.getSharedPreferences("LocalNotificationPreferences", MODE_PRIVATE);
@@ -5115,9 +5145,18 @@ public class GameActivity extends $${gameActivitySuperClass}$$ implements Surfac
 		}
 	}
 
-	public int AndroidThunkJava_LocalNotificationScheduleAtTime(String targetDateTime, boolean localTime, String title, String body, String action, String activationEvent) 
+	public int AndroidThunkJava_LocalNotificationScheduleAtTime(String targetDateTime, boolean localTime, String title, String body, String action, String activationEvent, int idOverride) 
 	{
-		int notificationID = LocalNotificationGetID(this);
+		int notificationID = idOverride;
+		if (notificationID == -1)
+		{
+			notificationID = LocalNotificationGetID(this);
+		}
+		else
+		{
+			//Make sure the Overriden id is a part of the current notification list
+			AddLocalNotificationID(this, notificationID);
+		}
 
 		if (!LocalNotificationScheduleAtTime(this, notificationID, targetDateTime, localTime, title, body, action, activationEvent))
 		{
diff --git a/Engine/Source/Runtime/Android/AndroidLocalNotification/Private/AndroidLocalNotification.cpp b/Engine/Source/Runtime/Android/AndroidLocalNotification/Private/AndroidLocalNotification.cpp
--- a/Engine/Source/Runtime/Android/AndroidLocalNotification/Private/AndroidLocalNotification.cpp
+++ b/Engine/Source/Runtime/Android/AndroidLocalNotification/Private/AndroidLocalNotification.cpp
@@ -86,8 +86,18 @@ void FAndroidLocalNotificationService::ClearAllLocalNotifications()
 int32 FAndroidLocalNotificationService::ScheduleLocalNotificationAtTime(const FDateTime& FireDateTime, bool LocalTime, const FText& Title, const FText& Body, const FText& Action, const FString& ActivationEvent)
 {
 #if USE_ANDROID_JNI
-	extern int32 AndroidThunkCpp_ScheduleLocalNotificationAtTime(const FDateTime& FireDateTime, bool LocalTime, const FText& Title, const FText& Body, const FText& Action, const FString& ActivationEvent);
-	return AndroidThunkCpp_ScheduleLocalNotificationAtTime(FireDateTime , LocalTime, Title, Body, Action, ActivationEvent);
+	extern int32 AndroidThunkCpp_ScheduleLocalNotificationAtTime(const FDateTime & FireDateTime, bool LocalTime, const FText & Title, const FText & Body, const FText & Action, const FString & ActivationEvent, int32 IdOverride);
+	return AndroidThunkCpp_ScheduleLocalNotificationAtTime(FireDateTime, LocalTime, Title, Body, Action, ActivationEvent, -1);
+#else
+	return -1;
+#endif
+}
+
+int32 FAndroidLocalNotificationService::ScheduleLocalNotificationAtTimeOverrideId(const FDateTime& FireDateTime, bool LocalTime, const FText& Title, const FText& Body, const FText& Action, const FString& ActivationEvent, int32 IdOverride)
+{
+#if USE_ANDROID_JNI
+	extern int32 AndroidThunkCpp_ScheduleLocalNotificationAtTime(const FDateTime & FireDateTime, bool LocalTime, const FText & Title, const FText & Body, const FText & Action, const FString & ActivationEvent, int32 IdOverride);
+	return AndroidThunkCpp_ScheduleLocalNotificationAtTime(FireDateTime, LocalTime, Title, Body, Action, ActivationEvent, IdOverride);
 #else
 	return -1;
 #endif
diff --git a/Engine/Source/Runtime/Android/AndroidLocalNotification/Public/AndroidLocalNotification.h b/Engine/Source/Runtime/Android/AndroidLocalNotification/Public/AndroidLocalNotification.h
--- a/Engine/Source/Runtime/Android/AndroidLocalNotification/Public/AndroidLocalNotification.h
+++ b/Engine/Source/Runtime/Android/AndroidLocalNotification/Public/AndroidLocalNotification.h
@@ -27,6 +27,7 @@ public:
 
 	virtual void ClearAllLocalNotifications();
 	virtual int32 ScheduleLocalNotificationAtTime(const FDateTime& FireDateTime, bool LocalTime, const FText& Title, const FText& Body, const FText& Action, const FString& ActivationEvent);
+	virtual int32 ScheduleLocalNotificationAtTimeOverrideId(const FDateTime& FireDateTime, bool LocalTime, const FText& Title, const FText& Body, const FText& Action, const FString& ActivationEvent, int32 IdOverride);
 	virtual int32 ScheduleLocalNotificationBadgeAtTime(const FDateTime& FireDateTime, bool LocalTime, const FString& ActivationEvent);
 	virtual void GetLaunchNotification(bool& NotificationLaunchedApp, FString& ActivationEvent, int32& FireDate);
 	virtual void SetLaunchNotification(FString const& ActivationEvent, int32 FireDate);
diff --git a/Engine/Source/Runtime/Engine/Public/LocalNotification.h b/Engine/Source/Runtime/Engine/Public/LocalNotification.h
--- a/Engine/Source/Runtime/Engine/Public/LocalNotification.h
+++ b/Engine/Source/Runtime/Engine/Public/LocalNotification.h
@@ -29,6 +29,19 @@ public:
 	*/
 	virtual int32 ScheduleLocalNotificationAtTime(const FDateTime& FireDateTime, bool LocalTime, const FText& Title, const FText& Body, const FText& Action, const FString& ActivationEvent) = 0;
 
+	/** Schedule a local notification at a specific time with a specific Id, inLocalTime specifies the current local time or if UTC time should be used
+	 * Currently only works for Android.
+	 * @param FireDateTime The time at which to fire the local notification
+	 * @param LocalTime If true the provided time is in the local timezone, if false it is in UTC
+	 * @param Title The title of the notification
+	 * @param Body The more detailed description of the notification
+	 * @param Action The text to be displayed on the slider controller
+	 * @param ActivationEvent A string that is passed in the delegate callback when the app is brought into the foreground from the user activating the notification
+	 * @param IdOverride A int that is passed to pin the notification to an Id. By doing this it will only show the latest notification sent from that ID.
+	 * @return A unique identifier used to cancel the notification
+	*/
+	virtual int32 ScheduleLocalNotificationAtTimeOverrideId(const FDateTime& FireDateTime, bool LocalTime, const FText& Title, const FText& Body, const FText& Action, const FString& ActivationEvent, int32 IdOverride) = 0;
+
 	/** Schedule a local notification badge at a specific time, inLocalTime specifies the current local time or if UTC time should be used
 	 * @param FireDateTime The time at which to fire the local notification
 	 * @param LocalTime If true the provided time is in the local timezone, if false it is in UTC
diff --git a/Engine/Source/Runtime/IOS/IOSLocalNotification/Private/IOSLocalNotification.cpp b/Engine/Source/Runtime/IOS/IOSLocalNotification/Private/IOSLocalNotification.cpp
--- a/Engine/Source/Runtime/IOS/IOSLocalNotification/Private/IOSLocalNotification.cpp
+++ b/Engine/Source/Runtime/IOS/IOSLocalNotification/Private/IOSLocalNotification.cpp
@@ -178,6 +178,11 @@ int32 FIOSLocalNotificationService::ScheduleLocalNotificationAtTime(const FDateT
 #endif
 }
 
+int32 FIOSLocalNotificationService::ScheduleLocalNotificationAtTimeOverrideId(const FDateTime& FireDateTime, bool LocalTime, const FText& Title, const FText& Body, const FText& Action, const FString& ActivationEvent, int32 IdOverride)
+{
+	return ScheduleLocalNotificationAtTime(FireDateTime, LocalTime, Title, Body, Action, ActivationEvent);
+}
+
 int32 FIOSLocalNotificationService::ScheduleLocalNotificationBadgeAtTime(const FDateTime& FireDateTime, bool LocalTime, const FString& ActivationEvent)
 {
 #if !PLATFORM_TVOS
diff --git a/Engine/Source/Runtime/IOS/IOSLocalNotification/Public/IOSLocalNotification.h b/Engine/Source/Runtime/IOS/IOSLocalNotification/Public/IOSLocalNotification.h
--- a/Engine/Source/Runtime/IOS/IOSLocalNotification/Public/IOSLocalNotification.h
+++ b/Engine/Source/Runtime/IOS/IOSLocalNotification/Public/IOSLocalNotification.h
@@ -22,6 +22,7 @@ public:
 	virtual ~FIOSLocalNotificationService() {}
 	virtual void ClearAllLocalNotifications();
 	virtual int32 ScheduleLocalNotificationAtTime(const FDateTime& FireDateTime, bool LocalTime, const FText& Title, const FText& Body, const FText& Action, const FString& ActivationEvent);
+	virtual int32 ScheduleLocalNotificationAtTimeOverrideId(const FDateTime& FireDateTime, bool LocalTime, const FText& Title, const FText& Body, const FText& Action, const FString& ActivationEvent, int32 IdOverride);
 	virtual int32 ScheduleLocalNotificationBadgeAtTime(const FDateTime& FireDateTime, bool LocalTime, const FString& ActivationEvent);
 	virtual void GetLaunchNotification(bool& NotificationLaunchedApp, FString& ActivationEvent, int32& FireDate);
 	virtual void SetLaunchNotification(FString const& ActivationEvent, int32 FireDate);
diff --git a/Engine/Source/Runtime/Launch/Private/Android/AndroidJNI.cpp b/Engine/Source/Runtime/Launch/Private/Android/AndroidJNI.cpp
--- a/Engine/Source/Runtime/Launch/Private/Android/AndroidJNI.cpp
+++ b/Engine/Source/Runtime/Launch/Private/Android/AndroidJNI.cpp
@@ -106,7 +106,7 @@ void FJavaWrapper::FindClassesAndMethods(JNIEnv* Env)
 	AndroidThunkJava_GetMetaDataString = FindMethod(Env, GameActivityClassID, "AndroidThunkJava_GetMetaDataString", "(Ljava/lang/String;)Ljava/lang/String;", bIsOptional);
 	AndroidThunkJava_SetSustainedPerformanceMode = FindMethod(Env, GameActivityClassID, "AndroidThunkJava_SetSustainedPerformanceMode", "(Z)V", bIsOptional);
 	AndroidThunkJava_ShowHiddenAlertDialog = FindMethod(Env, GameActivityClassID, "AndroidThunkJava_ShowHiddenAlertDialog", "()V", bIsOptional);
-	AndroidThunkJava_LocalNotificationScheduleAtTime = FindMethod(Env, GameActivityClassID, "AndroidThunkJava_LocalNotificationScheduleAtTime", "(Ljava/lang/String;ZLjava/lang/String;Ljava/lang/String;Ljava/lang/String;Ljava/lang/String;)I", bIsOptional);
+	AndroidThunkJava_LocalNotificationScheduleAtTime = FindMethod(Env, GameActivityClassID, "AndroidThunkJava_LocalNotificationScheduleAtTime", "(Ljava/lang/String;ZLjava/lang/String;Ljava/lang/String;Ljava/lang/String;Ljava/lang/String;I)I", bIsOptional);
 	AndroidThunkJava_LocalNotificationClearAll = FindMethod(Env, GameActivityClassID, "AndroidThunkJava_LocalNotificationClearAll", "()V", bIsOptional);
 	AndroidThunkJava_LocalNotificationExists = FindMethod(Env, GameActivityClassID, "AndroidThunkJava_LocalNotificationExists", "(I)Z", bIsOptional);
 	AndroidThunkJava_LocalNotificationGetLaunchNotification = FindMethod(Env, GameActivityClassID, "AndroidThunkJava_LocalNotificationGetLaunchNotification", "()Lcom/epicgames/unreal/GameActivity$LaunchNotification;", bIsOptional);
@@ -1482,7 +1482,7 @@ void AndroidThunkCpp_SetDesiredViewSize(int32 Width, int32 Height)
 // #endif
 }
 
-int32 AndroidThunkCpp_ScheduleLocalNotificationAtTime(const FDateTime& FireDateTime, bool LocalTime, const FText& Title, const FText& Body, const FText& Action, const FString& ActivationEvent)
+int32 AndroidThunkCpp_ScheduleLocalNotificationAtTime(const FDateTime& FireDateTime, bool LocalTime, const FText& Title, const FText& Body, const FText& Action, const FString& ActivationEvent, int32 IdOverride)
 {
 	//Convert FireDateTime to yyyy-MM-dd HH:mm:ss in order to pass to java
 	FString FireDateTimeFormatted = FString::FromInt(FireDateTime.GetYear()) + "-" + FString::FromInt(FireDateTime.GetMonth()) + "-" + FString::FromInt(FireDateTime.GetDay()) + " " + FString::FromInt(FireDateTime.GetHour()) + ":" + FString::FromInt(FireDateTime.GetMinute()) + ":" + FString::FromInt(FireDateTime.GetSecond());
@@ -1496,7 +1496,7 @@ int32 AndroidThunkCpp_ScheduleLocalNotificationAtTime(const FDateTime& FireDateT
 		auto jAction = FJavaHelper::ToJavaString(Env, Action.ToString());
 		auto jActivationEvent = FJavaHelper::ToJavaString(Env, ActivationEvent);
 		
-		return FJavaWrapper::CallIntMethod(Env, FJavaWrapper::GameActivityThis, FJavaWrapper::AndroidThunkJava_LocalNotificationScheduleAtTime, *jFireDateTime, LocalTime, *jTitle, *jBody, *jAction, *jActivationEvent);
+		return FJavaWrapper::CallIntMethod(Env, FJavaWrapper::GameActivityThis, FJavaWrapper::AndroidThunkJava_LocalNotificationScheduleAtTime, *jFireDateTime, LocalTime, *jTitle, *jBody, *jAction, *jActivationEvent, IdOverride);
 	}
 	
 	return -1;
diff --git a/Engine/Source/Runtime/Online/BackgroundHTTP/Private/BackgroundHttpNotificationObject.cpp b/Engine/Source/Runtime/Online/BackgroundHTTP/Private/BackgroundHttpNotificationObject.cpp
--- a/Engine/Source/Runtime/Online/BackgroundHTTP/Private/BackgroundHttpNotificationObject.cpp
+++ b/Engine/Source/Runtime/Online/BackgroundHTTP/Private/BackgroundHttpNotificationObject.cpp
@@ -11,11 +11,11 @@
 #include "LocalNotification.h"
 
 FBackgroundHttpNotificationObject::FBackgroundHttpNotificationObject(FText InNotificationTitle, FText InNotificationBody, FText InNotificationAction, const FString& InNotificationActivationString, bool InNotifyOnlyOnFullSuccess)
-	: FBackgroundHttpNotificationObject(InNotificationTitle, InNotificationBody, InNotificationAction, InNotificationActivationString, InNotifyOnlyOnFullSuccess, true)
+	: FBackgroundHttpNotificationObject(InNotificationTitle, InNotificationBody, InNotificationAction, InNotificationActivationString, InNotifyOnlyOnFullSuccess, true, -1)
 {
 }
 
-FBackgroundHttpNotificationObject::FBackgroundHttpNotificationObject(FText InNotificationTitle, FText InNotificationBody, FText InNotificationAction, const FString& InNotificationActivationString, bool InNotifyOnlyOnFullSuccess, bool InbOnlySendNotificationInBackground)
+FBackgroundHttpNotificationObject::FBackgroundHttpNotificationObject(FText InNotificationTitle, FText InNotificationBody, FText InNotificationAction, const FString& InNotificationActivationString, bool InNotifyOnlyOnFullSuccess, bool InbOnlySendNotificationInBackground, int32 InIdOverride)
 	: NotificationTitle(InNotificationTitle)
     , NotificationAction(InNotificationAction)
     , NotificationBody(InNotificationBody)
@@ -24,6 +24,7 @@ FBackgroundHttpNotificationObject::FBackgroundHttpNotificationObject(FText InNot
     , bNotifyOnlyOnFullSuccess(InNotifyOnlyOnFullSuccess)
 	, bIsInBackground(false)
 	, NumFailedDownloads(0)
+	, IdOverride(InIdOverride)
     , PlatformNotificationService(nullptr)
     , OnApp_EnteringForegroundHandle()
 	, OnApp_EnteringBackgroundHandle()
@@ -87,7 +88,7 @@ FBackgroundHttpNotificationObject::~FBackgroundHttpNotificationObject()
 
 			if (nullptr != PlatformNotificationService)
 			{
-				PlatformNotificationService->ScheduleLocalNotificationAtTime(TargetTime, true, NotificationTitle, NotificationBody, NotificationAction, NotificationActivationString);
+				PlatformNotificationService->ScheduleLocalNotificationAtTimeOverrideId(TargetTime, true, NotificationTitle, NotificationBody, NotificationAction, NotificationActivationString, IdOverride);
 			}
 		}
 	}
diff --git a/Engine/Source/Runtime/Online/BackgroundHTTP/Public/BackgroundHttpNotificationObject.h b/Engine/Source/Runtime/Online/BackgroundHTTP/Public/BackgroundHttpNotificationObject.h
--- a/Engine/Source/Runtime/Online/BackgroundHTTP/Public/BackgroundHttpNotificationObject.h
+++ b/Engine/Source/Runtime/Online/BackgroundHTTP/Public/BackgroundHttpNotificationObject.h
@@ -14,7 +14,7 @@ class BACKGROUNDHTTP_API FBackgroundHttpNotificationObject
 {
 public:
 	FBackgroundHttpNotificationObject(FText InNotificationTitle, FText InNotificationBody, FText InNotificationAction, const FString& InNotificationActivationString, bool InNotifyOnlyOnFullSuccess);
-	FBackgroundHttpNotificationObject(FText InNotificationTitle, FText InNotificationBody, FText InNotificationAction, const FString& InNotificationActivationString, bool InNotifyOnlyOnFullSuccess, bool bOnlySendNotificationInBackground);
+	FBackgroundHttpNotificationObject(FText InNotificationTitle, FText InNotificationBody, FText InNotificationAction, const FString& InNotificationActivationString, bool InNotifyOnlyOnFullSuccess, bool bOnlySendNotificationInBackground, int32 InIdOverride);
 	~FBackgroundHttpNotificationObject();
 
 	void NotifyOfDownloadResult(bool bWasSuccess);
@@ -35,6 +35,8 @@ private:
 	volatile bool bIsInBackground;
 	volatile int32 NumFailedDownloads;
 
+	int32 IdOverride;
+
 	class ILocalNotificationService* PlatformNotificationService;
 
 	FDelegateHandle OnApp_EnteringForegroundHandle;
