Add warning dialog on macOS [13.0~13.3), ask user to update to avoid dlopen crash

diff --git a/Engine/Source/Runtime/Launch/Private/Mac/LaunchMac.cpp b/Engine/Source/Runtime/Launch/Private/Mac/LaunchMac.cpp
--- a/Engine/Source/Runtime/Launch/Private/Mac/LaunchMac.cpp
+++ b/Engine/Source/Runtime/Launch/Private/Mac/LaunchMac.cpp
@@ -333,6 +333,43 @@ static int32 MacOSVersionCompare(const NSOperatingSystemVersion& VersionA, const
 		}
 	}
 #endif // WITH_EDITOR
+	
+#if !IS_MONOLITHIC
+	// UE-172403: dlopen crash on Ventura [13.0~13.3)
+	if (MacOSVersionCompare(CurrentSystemVersion, {13, 0, 0}) >= 0 && MacOSVersionCompare(CurrentSystemVersion, {13, 3, 0}) < 0)
+	{
+		CFDictionaryRef SessionDictionary = CGSessionCopyCurrentDictionary();
+		const bool bIsWindowServerAvailable = SessionDictionary != nullptr;
+		NSString* const kDialogSuppressKey = @"VenturaCrashWarningDialogSuppression";
+		if (![[NSUserDefaults standardUserDefaults] boolForKey:kDialogSuppressKey]
+			&& bIsWindowServerAvailable)
+		{
+			NSAlert* AlertPanel = [NSAlert new];
+			[AlertPanel setAlertStyle:NSAlertStyleCritical];
+			[AlertPanel setInformativeText:@"Due to a conflict between certain versions of macOS Ventura and Unreal Editor, it is recommended to update to macOS 13.3 or later. Continuing may cause the editor to crash during load."];
+			[AlertPanel setMessageText:@"Please update to latest macOS"];
+			[AlertPanel setShowsSuppressionButton:YES];
+			[AlertPanel addButtonWithTitle:@"Continue"];
+			[AlertPanel addButtonWithTitle:@"Quit"];
+			
+			auto Result = [AlertPanel runModal];
+			if (AlertPanel.suppressionButton.state == NSControlStateValueOn)
+			{
+				[[NSUserDefaults standardUserDefaults] setBool:YES forKey:kDialogSuppressKey];
+			}
+			[AlertPanel release];
+
+			CFRelease(SessionDictionary);
+			
+			if (Result == NSAlertSecondButtonReturn)
+			{
+				_Exit(1);
+			}
+		}
+
+		fprintf(stderr, "Due to a conflict between certain versions of macOS Ventura and Unreal Editor, it is recommended to update to macOS 13.3 or later. Continuing may cause the editor to crash during load.\n");
+	}
+#endif
 
 	//install the custom quit event handler
     NSAppleEventManager* appleEventManager = [NSAppleEventManager sharedAppleEventManager];
