Ignore focus events on Android to fix black screen on start problem with splitscreen mode

#android

diff --git a/Engine/Source/Runtime/Launch/Private/Android/LaunchAndroid.cpp b/Engine/Source/Runtime/Launch/Private/Android/LaunchAndroid.cpp
--- a/Engine/Source/Runtime/Launch/Private/Android/LaunchAndroid.cpp
+++ b/Engine/Source/Runtime/Launch/Private/Android/LaunchAndroid.cpp
@@ -1244,7 +1244,7 @@ FAppEventData::FAppEventData(ANativeWindow* WindowIn)
 static bool bAppIsActive_EventThread = false;
 
 
-// called when the app has focus + window + resume.
+// called when the app has window + resume.
 static void ActivateApp_EventThread()
 {
 	if (bAppIsActive_EventThread)
@@ -1283,7 +1283,7 @@ static void ActivateApp_EventThread()
 }
 
 extern void BlockRendering();
-// called whenever the app loses focus, loses window or pause.
+// called whenever the app loses loses window or pause.
 static void SuspendApp_EventThread()
 {
 	if (!bAppIsActive_EventThread)
@@ -1412,10 +1412,8 @@ static void OnAppCommandCB(struct android_app* app, int32_t cmd)
 		 * input focus.
 		 */
 		// if the app lost focus, avoid unnecessary processing (like monitoring the accelerometer)
-		bHasFocus = false;
-
+		// log it, but the actual event will be simulated later in APP_CMD_PAUSE
 		UE_LOG(LogAndroid, Log, TEXT("Case APP_CMD_LOST_FOCUS"));
-		FAppEventManager::GetInstance()->EnqueueAppEvent(APP_EVENT_STATE_WINDOW_LOST_FOCUS);
 
 		break;
 	case APP_CMD_GAINED_FOCUS:
@@ -1424,14 +1422,11 @@ static void OnAppCommandCB(struct android_app* app, int32_t cmd)
 		 * input focus.
 		 */
 
-		// remember gaining focus so we know any later pauses are not part of first startup
-		bDidGainFocus = true;
-		bHasFocus = true;
-
 		// bring back a certain functionality, like monitoring the accelerometer
+		// log it, but the actual event will be simulated later in APP_CMD_RESUME
 		UE_LOG(LogAndroid, Log, TEXT("Case APP_CMD_GAINED_FOCUS"));
-		FAppEventManager::GetInstance()->EnqueueAppEvent(APP_EVENT_STATE_WINDOW_GAINED_FOCUS);
 
+		// still check for a rare case needing activation
 		if (bHasWindow && bHasFocus && bIsResumed)
 		{
 			ActivateApp_EventThread();
@@ -1499,6 +1494,11 @@ static void OnAppCommandCB(struct android_app* app, int32_t cmd)
 		 * Command from main thread: the app's activity has been resumed.
 		 */
 		bIsResumed = true;
+
+		// assume focus on resume
+		bDidGainFocus = true;
+		bHasFocus = true;
+
 		if (bHasWindow && bHasFocus && bIsResumed)
 		{
 			ActivateApp_EventThread();
@@ -1507,6 +1507,9 @@ static void OnAppCommandCB(struct android_app* app, int32_t cmd)
 		UE_LOG(LogAndroid, Log, TEXT("Case APP_CMD_RESUME"));
 		FAppEventManager::GetInstance()->EnqueueAppEvent(APP_EVENT_STATE_ON_RESUME);
 
+		// trigger focus
+		FAppEventManager::GetInstance()->EnqueueAppEvent(APP_EVENT_STATE_WINDOW_GAINED_FOCUS);
+
 		/*
 		* On the initial loading the restart method must be called immediately
 		* in order to restart the app if the startup movie was playing
@@ -1524,6 +1527,10 @@ static void OnAppCommandCB(struct android_app* app, int32_t cmd)
 		FPlatformMisc::LowLevelOutputDebugStringf(TEXT("Case APP_CMD_PAUSE"));
 		UE_LOG(LogAndroid, Log, TEXT("Case APP_CMD_PAUSE"));
 
+		// simulate lost focus
+		FAppEventManager::GetInstance()->EnqueueAppEvent(APP_EVENT_STATE_WINDOW_LOST_FOCUS);
+		bHasFocus = false;
+
 		// Ignore pause command for Oculus if the window hasn't been initialized to prevent halting initial load
 		// if the headset is not active
 		if (!bHasWindow && FAndroidMisc::GetDeviceMake() == FString("Oculus"))
