LaunchEngineLoop log UE-CmdLineArgs

diff --git a/Engine/Source/Runtime/Launch/Private/LaunchEngineLoop.cpp b/Engine/Source/Runtime/Launch/Private/LaunchEngineLoop.cpp
--- a/Engine/Source/Runtime/Launch/Private/LaunchEngineLoop.cpp
+++ b/Engine/Source/Runtime/Launch/Private/LaunchEngineLoop.cpp
@@ -1924,6 +1924,8 @@ int32 FEngineLoop::PreInitPreStartupScreen(const TCHAR* CmdLine)
 	FString Env = FPlatformMisc::GetEnvironmentVariable(TEXT("UE-CmdLineArgs")).TrimStart();
 	if (Env.Len())
 	{
+		UE_LOG(LogInit, Log, TEXT("Inserting commandline from UE-CmdLineArgs: %s"), *Env);
+
 		// Append the command line environment after inserting a space as we can't set it in the
 		// environment.
 		FCommandLine::Append(TEXT(" -EnvAfterHere "));
