Update iOS suspend/resume logic to not manually force the render thread to stop/start, in line with Android

diff --git a/Engine/Source/Runtime/Launch/Private/IOS/LaunchIOS.cpp b/Engine/Source/Runtime/Launch/Private/IOS/LaunchIOS.cpp
--- a/Engine/Source/Runtime/Launch/Private/IOS/LaunchIOS.cpp
+++ b/Engine/Source/Runtime/Launch/Private/IOS/LaunchIOS.cpp
@@ -472,8 +472,6 @@ void FAppEntry::Init()
 #if BUILD_EMBEDDED_APP
 static bool GWasTickSuspended = false;
 static double GPreviousSuspendTime = FPlatformTime::Seconds();
-#else
-static FSuspendRenderingThread* SuspendThread = NULL;
 #endif
 
 void FAppEntry::Tick()
@@ -484,13 +482,6 @@ void FAppEntry::Tick()
 		FPlatformProcess::SetRealTimeMode();
 		GWasTickSuspended = false;
 	}
-#else
-	if (SuspendThread != NULL)
-	{
-		delete SuspendThread;
-		SuspendThread = NULL;
-		FPlatformProcess::SetRealTimeMode();
-	}
 #endif
     
 	if (AudioContextResumeTime != 0)
@@ -523,11 +514,6 @@ void FAppEntry::SuspendTick()
 	// allow for some background processing
 	FEmbeddedCommunication::TickGameThread(DeltaTime);
 	FCoreDelegates::MobileBackgroundTickDelegate.Broadcast(DeltaTime);
-#else
-	if (!SuspendThread)
-	{
-		SuspendThread = new FSuspendRenderingThread(true);
-	}
 #endif
 
 	FPlatformProcess::Sleep(0.1f);
