Load ScriptableEditorWidgets even earlier, it was refactored out of UMGEditor and so to maintain backwards compatability we must load it here (UMG transitively loads UMGEditor)

diff --git a/Engine/Source/Runtime/Launch/Private/LaunchEngineLoop.cpp b/Engine/Source/Runtime/Launch/Private/LaunchEngineLoop.cpp
--- a/Engine/Source/Runtime/Launch/Private/LaunchEngineLoop.cpp
+++ b/Engine/Source/Runtime/Launch/Private/LaunchEngineLoop.cpp
@@ -4535,6 +4535,9 @@ bool FEngineLoop::LoadStartupCoreModules()
 	// In dedicated server builds with the editor, we need to load UMG/UMGEditor for compiling blueprints.
 	// UMG must be loaded for runtime and cooking.
 	FModuleManager::Get().LoadModule("UMG");
+	// ScriptableEditorWidgets was refactored out of UMG, load it now so that we don't break existing
+	// projets that are not listing it in their uproject or uplugin:
+	FModuleManager::Get().LoadModule("ScriptableEditorWidgets");
 #else
 	if ( !IsRunningDedicatedServer() )
 	{
