Fix Target Device Mobile Preview

diff --git a/Engine/Source/Runtime/Launch/Private/LaunchEngineLoop.cpp b/Engine/Source/Runtime/Launch/Private/LaunchEngineLoop.cpp
--- a/Engine/Source/Runtime/Launch/Private/LaunchEngineLoop.cpp
+++ b/Engine/Source/Runtime/Launch/Private/LaunchEngineLoop.cpp
@@ -6681,7 +6681,7 @@ bool FEngineLoop::AppInit( )
 		FPIEPreviewDeviceModule* PIEPreviewDeviceProfileSelectorModule = FModuleManager::LoadModulePtr<FPIEPreviewDeviceModule>("PIEPreviewDeviceProfileSelector");
 		if (PIEPreviewDeviceProfileSelectorModule)
 		{
-			Scalability::ChangeScalabilityPreviewPlatform(PIEPreviewDeviceProfileSelectorModule->GetPreviewPlatformName(), GEditor->GetActiveShaderPlatform());
+			Scalability::ChangeScalabilityPreviewPlatform(PIEPreviewDeviceProfileSelectorModule->GetPreviewPlatformName(), GShaderPlatformForFeatureLevel[ERHIFeatureLevel::ES3_1]);
 		}
 	}
 #endif
