I/O - added on demand I/O store modue to engine

diff --git a/Engine/Source/Runtime/Engine/Engine.Build.cs b/Engine/Source/Runtime/Engine/Engine.Build.cs
--- a/Engine/Source/Runtime/Engine/Engine.Build.cs
+++ b/Engine/Source/Runtime/Engine/Engine.Build.cs
@@ -105,7 +105,8 @@ public class Engine : ModuleRules
 				"AudioExtensions",
 				"DeveloperSettings",
 				"AudioLinkCore",
-				"CookOnTheFly"
+				"CookOnTheFly",
+				"IoStoreOnDemand"
 			}
 		);
 
diff --git a/Engine/Source/Runtime/Launch/Private/LaunchEngineLoop.cpp b/Engine/Source/Runtime/Launch/Private/LaunchEngineLoop.cpp
--- a/Engine/Source/Runtime/Launch/Private/LaunchEngineLoop.cpp
+++ b/Engine/Source/Runtime/Launch/Private/LaunchEngineLoop.cpp
@@ -2202,6 +2202,11 @@ int32 FEngineLoop::PreInitPreStartupScreen(const TCHAR* CmdLine)
 		}
 	}
 
+#if USE_IO_DISPATCHER
+	// Initialize on demand I/O dispatcher backend
+	FModuleManager::Get().LoadModule("IoStoreOnDemand");
+#endif
+
 	// 
 #if PLATFORM_DESKTOP && !IS_MONOLITHIC
 	{
