Forward key pressed to UI if virtual keyboard is shown on the screen. Required to make hardware keyboard pass presses to EditText field

diff --git a/Engine/Source/Runtime/Launch/Private/Android/AndroidJNI.cpp b/Engine/Source/Runtime/Launch/Private/Android/AndroidJNI.cpp
--- a/Engine/Source/Runtime/Launch/Private/Android/AndroidJNI.cpp
+++ b/Engine/Source/Runtime/Launch/Private/Android/AndroidJNI.cpp
@@ -666,7 +666,7 @@ JNI_METHOD void Java_com_epicgames_unreal_GameActivity_nativeVirtualKeyboardVisi
 	}
 }
 
-bool AndroidThunkCpp_IsVirtuaKeyboardShown()
+bool AndroidThunkCpp_IsVirtualKeyboardShown()
 {
 	return GVirtualKeyboardShown;
 }
diff --git a/Engine/Source/Runtime/Launch/Private/Android/LaunchAndroid.cpp b/Engine/Source/Runtime/Launch/Private/Android/LaunchAndroid.cpp
--- a/Engine/Source/Runtime/Launch/Private/Android/LaunchAndroid.cpp
+++ b/Engine/Source/Runtime/Launch/Private/Android/LaunchAndroid.cpp
@@ -161,7 +161,7 @@ static FAutoConsoleVariableRef CVarEnableHWKeyboard(
 extern void AndroidThunkCpp_InitHMDs();
 extern void AndroidThunkCpp_ShowConsoleWindow();
 extern bool AndroidThunkCpp_VirtualInputIgnoreClick(int, int);
-extern bool AndroidThunkCpp_IsVirtuaKeyboardShown();
+extern bool AndroidThunkCpp_IsVirtualKeyboardShown();
 extern bool AndroidThunkCpp_IsWebViewShown();
 extern void AndroidThunkCpp_RestartApplication(const FString& IntentString);
 extern FString AndroidThunkCpp_GetIntentExtrasString(const FString& Key);
@@ -1068,7 +1068,7 @@ static int32_t HandleInputCB(struct android_app* app, AInputEvent* event)
 			}
 			FPlatformRect ScreenRect = FAndroidWindow::GetScreenRect(true);
 
-			if (AndroidThunkCpp_IsVirtuaKeyboardShown() && (type == TouchBegan || type == TouchMoved))
+			if (AndroidThunkCpp_IsVirtualKeyboardShown() && (type == TouchBegan || type == TouchMoved))
 			{
 				//ignore key down events when the native input was clicked or when the keyboard animation is playing
 				if (TryIgnoreClick(event, actionPointer))
@@ -1171,7 +1171,6 @@ static int32_t HandleInputCB(struct android_app* app, AInputEvent* event)
 		//Trap codes handled as possible gamepad events
 		if (device >= 0 && ValidGamepadKeyCodes.Contains(keyCode))
 		{
-			
 			bool down = AKeyEvent_getAction(event) != AKEY_EVENT_ACTION_UP;
 			FAndroidInputInterface::JoystickButtonEvent(device, keyCode, down);
 			FPlatformMisc::LowLevelOutputDebugStringf(TEXT("Received gamepad button: %d"), keyCode);
@@ -1186,7 +1185,15 @@ static int32_t HandleInputCB(struct android_app* app, AInputEvent* event)
 				return 0;
 			}
 
-			if (bSoftKey || GAndroidEnableHardwareKeyboard || AlwaysAllowedKeyCodes.Contains(keyCode))
+			// forward key presses to UI if needed 
+			if (AndroidThunkCpp_IsVirtualKeyboardShown())
+			{
+				return 0;
+			}
+
+			if (bSoftKey ||
+				(GAndroidEnableHardwareKeyboard && !AndroidThunkCpp_IsVirtualKeyboardShown()) ||
+				AlwaysAllowedKeyCodes.Contains(keyCode))
 			{
 				FDeferredAndroidMessage Message;
 
diff --git a/Engine/Source/Runtime/Launch/Public/Android/AndroidJNI.h b/Engine/Source/Runtime/Launch/Public/Android/AndroidJNI.h
--- a/Engine/Source/Runtime/Launch/Public/Android/AndroidJNI.h
+++ b/Engine/Source/Runtime/Launch/Public/Android/AndroidJNI.h
@@ -79,7 +79,7 @@ public:
 	static jmethodID AndroidThunkJava_DisableScreenCapture;
 
 	static jmethodID AndroidThunkCpp_VirtualInputIgnoreClick;
-	static jmethodID AndroidThunkCpp_IsVirtuaKeyboardShown;
+	static jmethodID AndroidThunkCpp_IsVirtualKeyboardShown;
 	static jmethodID AndroidThunkCpp_IsWebViewShown;
 
 	// InputDeviceInfo member field ids
