Ensure that when we set "ShowErrorCount=False" on an engine commandlet, we still return a non-zero error code when errors are written to the log... even if the commandlet itself returns 0. This replicates the behaviour that you get when ShowErrorCount=True

diff --git a/Engine/Source/Runtime/Launch/Private/LaunchEngineLoop.cpp b/Engine/Source/Runtime/Launch/Private/LaunchEngineLoop.cpp
--- a/Engine/Source/Runtime/Launch/Private/LaunchEngineLoop.cpp
+++ b/Engine/Source/Runtime/Launch/Private/LaunchEngineLoop.cpp
@@ -4220,6 +4220,12 @@ int32 FEngineLoop::PreInitPostStartupScreen(const TCHAR* CmdLine)
 			}
 			else
 			{
+				// Even when we're not displaying error/warning summaries in the log, still return an non-zero code if errors were logged
+				if ((ErrorLevel == 0) && (GWarn->GetNumErrors() > 0))
+				{
+					ErrorLevel = 1;
+				}
+
 				UE_LOG(LogInit, Display, TEXT("Finished.") );
 			}
 
