Bumping Engine version to 5.4

#rnx

diff --git a/Engine/Build/Build.version b/Engine/Build/Build.version
--- a/Engine/Build/Build.version
+++ b/Engine/Build/Build.version
@@ -1,6 +1,6 @@
 {
 	"MajorVersion": 5,
-	"MinorVersion": 3,
+	"MinorVersion": 4,
 	"PatchVersion": 0,
 	"Changelist": 0,
 	"CompatibleChangelist": 0,
diff --git a/Engine/Source/Runtime/Launch/Resources/Version.h b/Engine/Source/Runtime/Launch/Resources/Version.h
--- a/Engine/Source/Runtime/Launch/Resources/Version.h
+++ b/Engine/Source/Runtime/Launch/Resources/Version.h
@@ -56,14 +56,14 @@
 // newer than a 5.11.* version, regardless of the changelist that it was built with)
 // When updating these, also update the static_assert below
 #define ENGINE_MAJOR_VERSION	5
-#define ENGINE_MINOR_VERSION	3
+#define ENGINE_MINOR_VERSION	4
 #define ENGINE_PATCH_VERSION	0
 
 // If this static_assert fires then Version.h has been updated without updating this code.
 // This line exists to cause conflicts when merging Version.h between streams so if one stream updates
 // the minor version while another updates the patch, you do not get a silently combined version that
 // was unintended.
-static_assert(ENGINE_MAJOR_VERSION == 5 && ENGINE_MINOR_VERSION == 3 && ENGINE_PATCH_VERSION == 0); //-V501
+static_assert(ENGINE_MAJOR_VERSION == 5 && ENGINE_MINOR_VERSION == 4 && ENGINE_PATCH_VERSION == 0); //-V501
 
 // Macros for encoding strings
 #define VERSION_TEXT(x) TEXT(x)
