Added UCommandlet::UseCommandletResultAsExitCode
* If enabled, process return code will always be the commandlet result code
* Used automatically when lauching WorldPartition builders from the editor (was the previous behavior before 25946840)

diff --git a/Engine/Source/Editor/UnrealEd/Private/Commandlets/WorldPartitionBuilderCommandlet.cpp b/Engine/Source/Editor/UnrealEd/Private/Commandlets/WorldPartitionBuilderCommandlet.cpp
--- a/Engine/Source/Editor/UnrealEd/Private/Commandlets/WorldPartitionBuilderCommandlet.cpp
+++ b/Engine/Source/Editor/UnrealEd/Private/Commandlets/WorldPartitionBuilderCommandlet.cpp
@@ -71,8 +71,8 @@ int32 UWorldPartitionBuilderCommandlet::Main(const FString& Params)
 
 	if (Switches.Contains(TEXT("RunningFromUnrealEd")))
 	{
-		ShowErrorCount = false;	// This has the side effect of making the process return code match the return code of the commandlet
-		FastExit = true;		// Faster exit which avoids crash during shutdown. The engine isn't shutdown cleanly.
+		UseCommandletResultAsExitCode = true;	// The process return code will match the return code of the commandlet
+		FastExit = true;						// Faster exit which avoids crash during shutdown. The engine isn't shutdown cleanly.
 	}
 
 	ICollectionManager& CollectionManager = FModuleManager::LoadModuleChecked<FCollectionManagerModule>("CollectionManager").Get();
diff --git a/Engine/Source/Runtime/Engine/Classes/Commandlets/Commandlet.h b/Engine/Source/Runtime/Engine/Classes/Commandlets/Commandlet.h
--- a/Engine/Source/Runtime/Engine/Classes/Commandlets/Commandlet.h
+++ b/Engine/Source/Runtime/Engine/Classes/Commandlets/Commandlet.h
@@ -91,6 +91,10 @@ class UCommandlet : public UObject
 	UPROPERTY()
 	uint32 FastExit:1;
 
+	/** Whether to use the commandlet result as the exit code, even if errors were logged during the whole engine execution */
+	UPROPERTY()
+	uint32 UseCommandletResultAsExitCode: 1;
+
 	/**
 	 * Entry point for your commandlet
 	 *
diff --git a/Engine/Source/Runtime/Engine/Private/Commandlets/Commandlet.cpp b/Engine/Source/Runtime/Engine/Private/Commandlets/Commandlet.cpp
--- a/Engine/Source/Runtime/Engine/Private/Commandlets/Commandlet.cpp
+++ b/Engine/Source/Runtime/Engine/Private/Commandlets/Commandlet.cpp
@@ -24,6 +24,7 @@ UCommandlet::UCommandlet(const FObjectInitializer& ObjectInitializer)
 	ShowErrorCount = true;
 	ShowProgress = true;
 	FastExit = false;
+	UseCommandletResultAsExitCode = false;
 }
 
 
diff --git a/Engine/Source/Runtime/Launch/Private/LaunchEngineLoop.cpp b/Engine/Source/Runtime/Launch/Private/LaunchEngineLoop.cpp
--- a/Engine/Source/Runtime/Launch/Private/LaunchEngineLoop.cpp
+++ b/Engine/Source/Runtime/Launch/Private/LaunchEngineLoop.cpp
@@ -4219,19 +4219,21 @@ int32 FEngineLoop::PreInitPostStartupScreen(const TCHAR* CmdLine)
 				{
 					SET_WARN_COLOR(COLOR_RED);
 					UE_LOG(LogInit, Display, TEXT("Failure - %d error(s), %d warning(s)"), AllErrors.Num(), AllWarnings.Num() );
-					ErrorLevel = 1;
 				}
 				CLEAR_WARN_COLOR();
 			}
 			else
 			{
-				// Even when we're not displaying error/warning summaries in the log, still return an non-zero code if errors were logged
+				UE_LOG(LogInit, Display, TEXT("Finished."));
+			}
+
+			// Return an non-zero code if errors were logged and UseCommandletResultAsExitCode is false
+			if (!Commandlet->UseCommandletResultAsExitCode)
+			{
 				if ((ErrorLevel == 0) && (GWarn->GetNumErrors() > 0))
 				{
 					ErrorLevel = 1;
 				}
-
-				UE_LOG(LogInit, Display, TEXT("Finished.") );
 			}
 
 			double CommandletExecutionTime = FPlatformTime::Seconds() - CommandletExecutionStartTime;
