Show the current frame number in Insights as a hierarchical event (shows the CSV frame number if the CSV profiler is running).

diff --git a/Engine/Source/Runtime/Launch/Private/LaunchEngineLoop.cpp b/Engine/Source/Runtime/Launch/Private/LaunchEngineLoop.cpp
--- a/Engine/Source/Runtime/Launch/Private/LaunchEngineLoop.cpp
+++ b/Engine/Source/Runtime/Launch/Private/LaunchEngineLoop.cpp
@@ -290,6 +290,12 @@ class FFeedbackContext;
 #include "IOS/IOSAppDelegate.h"
 #endif
 
+#if CPUPROFILERTRACE_ENABLED
+UE_TRACE_EVENT_BEGIN(Cpu, Frame, NoSync)
+	UE_TRACE_EVENT_FIELD(UE::Trace::WideString, Name)
+UE_TRACE_EVENT_END()
+#endif // CPUPROFILERTRACE_ENABLED
+
 bool GIsConsoleExecutable = false;
 
 int32 GUseDisregardForGCOnDedicatedServers = 1;
@@ -5466,7 +5472,7 @@ static inline void BeginFrameRenderThread(FRHICommandListImmediate& RHICmdList,
 #else
 	FPlatformMisc::BeginNamedEvent(FColor::Yellow, *FrameString);
 #endif
-#endif // ENABLE_NAMED_EVENTS
+#endif // ENABLE_NAMED_EVENTS 
 
 	RHICmdList.PushEvent(*FrameString, FColor::Green);
 #endif // !UE_BUILD_SHIPPING
@@ -5575,6 +5581,7 @@ void FEngineLoop::Tick()
 	FMoviePlayerProxy::BlockingForceFinished();
 	ensure(GetMoviePlayer()->IsLoadingFinished() && !GetMoviePlayer()->IsMovieCurrentlyPlaying());
 
+    // Frame profiling kickoff
 #if UE_EXTERNAL_PROFILING_ENABLED
 	FExternalProfiler* ActiveProfiler = FActiveExternalProfilerBase::GetActiveProfiler();
 	if (ActiveProfiler)
@@ -5588,37 +5595,37 @@ void FEngineLoop::Tick()
 	uint64 CurrentFrameCounter = GFrameCounter;
 
 #if ENABLE_NAMED_EVENTS
-	TCHAR IndexedFrameString[32] = { 0 };
-	const TCHAR* FrameString = nullptr;
+	bool bTraceCpuChannelEnabled = false;
 #if CPUPROFILERTRACE_ENABLED
-	const bool bTraceCpuChannelEnabled = UE_TRACE_CHANNELEXPR_IS_ENABLED(CpuChannel);
+	bTraceCpuChannelEnabled = UE_TRACE_CHANNELEXPR_IS_ENABLED(CpuChannel);
 	static FAutoNamedEventsToggler TraceNamedEventsToggler;
 	TraceNamedEventsToggler.Update(bTraceCpuChannelEnabled && UE::Trace::IsTracing());
-	if (bTraceCpuChannelEnabled)
+#endif
+
+	// Generate the Frame event string
+	TCHAR IndexedFrameString[32] = { 0 };
+	const TCHAR* FrameStringPrefix = bTraceCpuChannelEnabled ? TEXT(" ") : TEXT("Frame ");
+#if CSV_PROFILER
+	if (FCsvProfiler::Get()->IsCapturing())
 	{
-		FrameString = TEXT("FEngineLoop");
+		FCString::Snprintf(IndexedFrameString, 32, TEXT("Csv%s%d"), FrameStringPrefix, FCsvProfiler::Get()->GetCaptureFrameNumber());
 	}
 	else
 #endif
 	{
-#if PLATFORM_LIMIT_PROFILER_UNIQUE_NAMED_EVENTS
-		FrameString = TEXT("FEngineLoop");
-#else // PLATFORM_LIMIT_PROFILER_UNIQUE_NAMED_EVENTS
-#if CSV_PROFILER
-		if (FCsvProfiler::Get()->IsCapturing())
-		{
-			FCString::Snprintf(IndexedFrameString, 32, TEXT("CsvFrame %d"), FCsvProfiler::Get()->GetCaptureFrameNumber());
-		}
-		else
-#endif
-		{
-			FCString::Snprintf(IndexedFrameString, 32, TEXT("Frame %d"), CurrentFrameCounter);
-		}
-		FrameString = IndexedFrameString;
-#endif // !PLATFORM_LIMIT_PROFILER_UNIQUE_NAMED_EVENTS
+		FCString::Snprintf(IndexedFrameString, 32, TEXT("%s%d"), FrameStringPrefix, CurrentFrameCounter);
 	}
-	SCOPED_NAMED_EVENT_TCHAR(FrameString, FColor::Red);
+	const TCHAR* FrameString = IndexedFrameString;
+#if CPUPROFILERTRACE_ENABLED
+	UE_TRACE_LOG_SCOPED_T(Cpu, Frame, CpuChannel) 
+		<< Frame.Name(FrameString);
+#endif // CPUPROFILERTRACE_ENABLED
+
+#if PLATFORM_LIMIT_PROFILER_UNIQUE_NAMED_EVENTS
+	FrameString = TEXT("FEngineLoop");
 #endif
+	SCOPED_NAMED_EVENT_TCHAR_CONDITIONAL(FrameString, FColor::Red, !bTraceCpuChannelEnabled); 
+#endif //ENABLE_NAMED_EVENTS
 
 	// execute callbacks for cvar changes
 	{
