Added FCoreDelegates::OnCommandletPreMain and FCoreDelegates::OnCommandletPostMain

This allows things to hook in to provide additional initialization or shutdown logic for any commandlet that may be run.

FCoreDelegates::OnCommandletPreMain is akin to FCoreDelegates::OnFEngineLoopInitComplete (which doesn't run for commandlets).

diff --git a/Engine/Source/Runtime/Core/Private/Misc/CoreDelegates.cpp b/Engine/Source/Runtime/Core/Private/Misc/CoreDelegates.cpp
--- a/Engine/Source/Runtime/Core/Private/Misc/CoreDelegates.cpp
+++ b/Engine/Source/Runtime/Core/Private/Misc/CoreDelegates.cpp
@@ -102,6 +102,8 @@ FSimpleMulticastDelegate FCoreDelegates::OnOutputDevicesInit;
 FSimpleMulticastDelegate FCoreDelegates::OnPostEngineInit;
 FSimpleMulticastDelegate FCoreDelegates::OnAllModuleLoadingPhasesComplete;
 FSimpleMulticastDelegate FCoreDelegates::OnFEngineLoopInitComplete;
+FSimpleMulticastDelegate FCoreDelegates::OnCommandletPreMain;
+FSimpleMulticastDelegate FCoreDelegates::OnCommandletPostMain;
 FSimpleMulticastDelegate FCoreDelegates::OnExit;
 FSimpleMulticastDelegate FCoreDelegates::OnPreExit;
 FSimpleMulticastDelegate FCoreDelegates::OnEnginePreExit;
diff --git a/Engine/Source/Runtime/Core/Public/Misc/CoreDelegates.h b/Engine/Source/Runtime/Core/Public/Misc/CoreDelegates.h
--- a/Engine/Source/Runtime/Core/Public/Misc/CoreDelegates.h
+++ b/Engine/Source/Runtime/Core/Public/Misc/CoreDelegates.h
@@ -196,6 +196,12 @@ public:
 	// Called at the very end of engine initialization, right before the engine starts ticking. This is not called for commandlets
 	static CORE_API FSimpleMulticastDelegate OnFEngineLoopInitComplete;
 
+	// Called prior to running the Main function for commandlets
+	static FSimpleMulticastDelegate OnCommandletPreMain;
+
+	// Called after running the Main function for commandlets
+	static FSimpleMulticastDelegate OnCommandletPostMain;
+
 	// Called when the application is about to exit.
 	static CORE_API FSimpleMulticastDelegate OnExit;
 
diff --git a/Engine/Source/Runtime/Launch/Private/LaunchEngineLoop.cpp b/Engine/Source/Runtime/Launch/Private/LaunchEngineLoop.cpp
--- a/Engine/Source/Runtime/Launch/Private/LaunchEngineLoop.cpp
+++ b/Engine/Source/Runtime/Launch/Private/LaunchEngineLoop.cpp
@@ -4129,11 +4129,13 @@ int32 FEngineLoop::PreInitPostStartupScreen(const TCHAR* CmdLine)
 #if WITH_ENGINE
 			PRIVATE_GRunningCommandletClass = CommandletClass;
 #endif
+			FCoreDelegates::OnCommandletPreMain.Broadcast();
 			int32 ErrorLevel;
 			{
 				FTrackedActivityScope CommandletActivity(FTrackedActivity::GetEngineActivity(), *FString::Printf(TEXT("Running %s"), *Commandlet->GetName()), false, FTrackedActivity::ELight::Green);
 				ErrorLevel = Commandlet->Main(CommandletCommandLine);
 			}
+			FCoreDelegates::OnCommandletPostMain.Broadcast();
 #if WITH_ENGINE
 			PRIVATE_GRunningCommandletClass = nullptr;
 #endif
