Merged Android PGO Build shelf to codebase (Engine change)

#rnx

diff --git a/Engine/Source/Runtime/Launch/Private/Android/LaunchAndroid.cpp b/Engine/Source/Runtime/Launch/Private/Android/LaunchAndroid.cpp
--- a/Engine/Source/Runtime/Launch/Private/Android/LaunchAndroid.cpp
+++ b/Engine/Source/Runtime/Launch/Private/Android/LaunchAndroid.cpp
@@ -365,6 +365,17 @@ static void InitCommandLine()
 		FPlatformMisc::LowLevelOutputDebugStringf(TEXT("UE setprop appended: %s"), UTF8_TO_TCHAR(CommandLineSetpropAppend));
 	}
 #endif
+
+#ifdef UE_ANDROID_COMMAND_LINE_OVERRIDE
+	FCommandLine::Set(TEXT(UE_ANDROID_COMMAND_LINE_OVERRIDE));
+	FPlatformMisc::LowLevelOutputDebugStringf(TEXT("UE_ANDROID_COMMAND_LINE_OVERRIDE: %s"), TEXT(UE_ANDROID_COMMAND_LINE_OVERRIDE));
+#endif
+
+#ifdef UE_ANDROID_COMMAND_LINE_APPEND
+	FCommandLine::Append(TEXT(UE_ANDROID_COMMAND_LINE_APPEND));
+	FPlatformMisc::LowLevelOutputDebugStringf(TEXT("UE_ANDROID_COMMAND_LINE_APPEND: %s"), TEXT(UE_ANDROID_COMMAND_LINE_APPEND));
+#endif
+
 }
 
 extern void AndroidThunkCpp_DismissSplashScreen();
