Fix for leaking render objects on app exit caused by EngineLoop not processing the last set of pending cleanup objects it acquired

diff --git a/Engine/Source/Runtime/Launch/Private/LaunchEngineLoop.cpp b/Engine/Source/Runtime/Launch/Private/LaunchEngineLoop.cpp
--- a/Engine/Source/Runtime/Launch/Private/LaunchEngineLoop.cpp
+++ b/Engine/Source/Runtime/Launch/Private/LaunchEngineLoop.cpp
@@ -4957,6 +4957,8 @@ void FEngineLoop::Exit()
 	TRACE_CPUPROFILER_EVENT_SCOPE(FEngineLoop::Exit);
 	TRACE_BOOKMARK(TEXT("EngineLoop.Exit"));
 
+	ClearPendingCleanupObjects();
+
 	GIsRunning	= 0;
 	GLogConsole	= nullptr;
 
