Add additional test string for where the PreInit module might be located

diff --git a/Engine/Source/Runtime/Launch/Private/LaunchEngineLoop.cpp b/Engine/Source/Runtime/Launch/Private/LaunchEngineLoop.cpp
--- a/Engine/Source/Runtime/Launch/Private/LaunchEngineLoop.cpp
+++ b/Engine/Source/Runtime/Launch/Private/LaunchEngineLoop.cpp
@@ -2109,13 +2109,19 @@ int32 FEngineLoop::PreInitPreStartupScreen(const TCHAR* CmdLine)
 		ProjectBinariesRootDirectory = FPlatformMisc::ProjectDir();
 #if !IS_MONOLITHIC
 		FPlatformMisc::GetEngineAndProjectAbsoluteDirsFromExecutable(ProjectBinariesRootDirectory, EngineBinariesRootDirectory);
-
 		// Loading preinit module if it exists. PreInit module can contain things like decryption logic for pak files and things that is project specific but needs to initialize very early
-		FString FileName = FString::Printf(TEXT("%s-%sPreInit.%s"), FPlatformProcess::ExecutableName(), FApp::GetProjectName(), FPlatformProcess::GetModuleExtension());
-		FString PreInitModule = FPaths::Combine(ProjectBinariesRootDirectory, "Binaries", FPlatformProcess::GetBinariesSubdirectory(), FileName);
-		if (FPaths::FileExists(PreInitModule))
+		TArray<FString> FileNames = {
+			FString::Printf(TEXT("%s-%sPreInit.%s"), FPlatformProcess::ExecutableName(), FApp::GetProjectName(), FPlatformProcess::GetModuleExtension()),
+			FString::Printf(TEXT("%s-%sPreInit-%s-%s.%s"), *FApp::GetName(), FApp::GetProjectName(), FPlatformProcess::GetBinariesSubdirectory(), LexToString(FApp::GetBuildConfiguration()), FPlatformProcess::GetModuleExtension())
+		};
+		for (const FString& FileName : FileNames)
 		{
-			FPlatformProcess::GetDllHandle(*PreInitModule);
+			FString PreInitModule = FPaths::Combine(ProjectBinariesRootDirectory, "Binaries", FPlatformProcess::GetBinariesSubdirectory(), FileName);
+			if (FPaths::FileExists(PreInitModule))
+			{
+				FPlatformProcess::GetDllHandle(*PreInitModule);
+				break;
+			}
 		}
 #endif
 	}
