Fix invalid Log parameter passing when SwiftUI is enabled

diff --git a/Engine/Source/Runtime/Launch/Private/IOS/LaunchIOS.cpp b/Engine/Source/Runtime/Launch/Private/IOS/LaunchIOS.cpp
--- a/Engine/Source/Runtime/Launch/Private/IOS/LaunchIOS.cpp
+++ b/Engine/Source/Runtime/Launch/Private/IOS/LaunchIOS.cpp
@@ -652,7 +652,7 @@ void FSwiftAppBootstrap::KickoffWithCompositingLayer(CP_OBJECT_cp_layer_renderer
 	CGRect FirstViewport = [[AppDelegate.SwiftLayerViewports firstObject] CGRectValue];
 	
 	FPlatformMisc::LowLevelOutputDebugStringf(TEXT("Kicking off UE with Swift Layer. Commandline: %s\n"), *GSavedCommandLine);
-	FPlatformMisc::LowLevelOutputDebugStringf(TEXT("NumViews: %d, Full size = %f x %f\n", NumViews, FirstViewport.size.width, FirstViewport.size.height);
+	FPlatformMisc::LowLevelOutputDebugStringf(TEXT("NumViews: %d, Full size = %f x %f\n"), NumViews, FirstViewport.size.width, FirstViewport.size.height);
 
 	AppDelegate.SwiftLayer = Layer;
 }
