Fix a deadlock that occurs if an android device is backgrounded during engine init.

diff --git a/Engine/Source/Runtime/Launch/Private/Android/LaunchAndroid.cpp b/Engine/Source/Runtime/Launch/Private/Android/LaunchAndroid.cpp
--- a/Engine/Source/Runtime/Launch/Private/Android/LaunchAndroid.cpp
+++ b/Engine/Source/Runtime/Launch/Private/Android/LaunchAndroid.cpp
@@ -265,6 +265,8 @@ void InitHMDs()
 
 	AndroidThunkCpp_InitHMDs();
 
+	GHMDsInitialized = GHMDImplementations.IsEmpty();
+
 	while (!GHMDsInitialized)
 	{
 		FPlatformProcess::Sleep(0.01f);
