Flush rendering commands prior to clearing pending delete objects at shutdown.

diff --git a/Engine/Source/Runtime/Launch/Private/LaunchEngineLoop.cpp b/Engine/Source/Runtime/Launch/Private/LaunchEngineLoop.cpp
--- a/Engine/Source/Runtime/Launch/Private/LaunchEngineLoop.cpp
+++ b/Engine/Source/Runtime/Launch/Private/LaunchEngineLoop.cpp
@@ -6153,6 +6153,7 @@ void FEngineLoop::Tick()
 
 void FEngineLoop::ClearPendingCleanupObjects()
 {
+	FlushRenderingCommands();
 	delete PendingCleanupObjects;
 	PendingCleanupObjects = nullptr;
 }
