PR #10737: Guard C++ semantics in Version.h
#rnx

diff --git a/Engine/Source/Runtime/Launch/Resources/Version.h b/Engine/Source/Runtime/Launch/Resources/Version.h
--- a/Engine/Source/Runtime/Launch/Resources/Version.h
+++ b/Engine/Source/Runtime/Launch/Resources/Version.h
@@ -21,7 +21,7 @@
 //         FEngineVersion::CompatibleWith() normally uses ENGINE_COMPATIBLE_CL_VERSION for the changelist component and '0' for the patch component, and 
 //                                          indicates the baseline version of the engine that this build maintains strict binary compatibility with. 
 //                                          By default, this compatibility extends to assets, executable modules, and any network data transmitted between 
-//                                          two builds, and is used when creating patches and hotfixes that can be used interchangably with another build.
+//                                          two builds, and is used when creating patches and hotfixes that can be used interchangeably with another build.
 //                                          This should be used for versioning in the majority of cases in the engine.
 //
 //     Both the ENGINE_CURRENT_CL_VERSION and ENGINE_COMPATIBLE_CL_VERSION macros can be updated systemically by build systems using the UpdateLocalVersion 
@@ -35,7 +35,7 @@
 //  -  The *licensee object version* is provided for licensees to create their own one-way upgrade paths akin to the regular object version. Epic will never
 //     add entries to this enumeration. It is defined by the enum in ObjectVersion.h
 //
-//  -  Any number of *custom object version* objects may be registered to create orthoganal incrementing version numbers similar to the object version and 
+//  -  Any number of *custom object version* objects may be registered to create orthogonal incrementing version numbers similar to the object version and 
 //     licensee version enums (see FCustomVersion). Each one is registered with a GUID, ensuring uniqueness and allowing the FArchive to quickly store and 
 //     retrieve them without any context of what they represent. Custom versions may be created for individual projects, subsystems, or branches.
 //
@@ -63,7 +63,9 @@
 // This line exists to cause conflicts when merging Version.h between streams so if one stream updates
 // the minor version while another updates the patch, you do not get a silently combined version that
 // was unintended.
+#ifdef __cplusplus
 static_assert(ENGINE_MAJOR_VERSION == 5 && ENGINE_MINOR_VERSION == 4 && ENGINE_PATCH_VERSION == 0); //-V501
+#endif // __cplusplus
 
 // Macros for encoding strings
 #define VERSION_TEXT(x) TEXT(x)
