Android: Fix out of bounds read.

diff --git a/Engine/Source/Runtime/Launch/Private/Android/AndroidEventManager.cpp b/Engine/Source/Runtime/Launch/Private/Android/AndroidEventManager.cpp
--- a/Engine/Source/Runtime/Launch/Private/Android/AndroidEventManager.cpp
+++ b/Engine/Source/Runtime/Launch/Private/Android/AndroidEventManager.cpp
@@ -30,7 +30,7 @@ FAppEventManager* FAppEventManager::GetInstance()
 
 static const TCHAR* GetAppEventName(EAppEventState State)
 {
-	const TCHAR* Names[] = {
+	const TCHAR* Names[APP_EVENT_MAX_EVENTS] = {
 		TEXT("APP_EVENT_STATE_WINDOW_CREATED"),
 		TEXT("APP_EVENT_STATE_WINDOW_RESIZED"),
 		TEXT("APP_EVENT_STATE_WINDOW_CHANGED"),
@@ -46,15 +46,17 @@ static const TCHAR* GetAppEventName(EAppEventState State)
 		TEXT("APP_EVENT_STATE_SAVE_STATE"),
 		TEXT("APP_EVENT_STATE_APP_SUSPENDED"),
 		TEXT("APP_EVENT_STATE_APP_ACTIVATED"),
+		TEXT("APP_EVENT_STATE_SAFE_ZONE_UPDATED"),
 		TEXT("APP_EVENT_RUN_CALLBACK"),
 		};
 
+	static_assert(APP_EVENT_MAX_EVENTS == 17, "You must update this array if more app events are added!");
 
 	if (State == APP_EVENT_STATE_INVALID)
 	{
 		return TEXT("APP_EVENT_STATE_INVALID");
 	}
-	else if (State > APP_EVENT_RUN_CALLBACK || State < 0)
+	else if (State >= APP_EVENT_MAX_EVENTS || State < 0)
 	{
 		return TEXT("UnknownEAppEventStateValue");
 	}
diff --git a/Engine/Source/Runtime/Launch/Public/Android/AndroidEventManager.h b/Engine/Source/Runtime/Launch/Public/Android/AndroidEventManager.h
--- a/Engine/Source/Runtime/Launch/Public/Android/AndroidEventManager.h
+++ b/Engine/Source/Runtime/Launch/Public/Android/AndroidEventManager.h
@@ -30,6 +30,7 @@ enum EAppEventState
 	APP_EVENT_STATE_APP_ACTIVATED,
 	APP_EVENT_STATE_SAFE_ZONE_UPDATED,
 	APP_EVENT_RUN_CALLBACK,
+	APP_EVENT_MAX_EVENTS,
 	APP_EVENT_STATE_INVALID = -1,
 };
 
