Additional insights CPU profiling scopes

diff --git a/Engine/Source/Runtime/AssetRegistry/Private/AssetDataGatherer.cpp b/Engine/Source/Runtime/AssetRegistry/Private/AssetDataGatherer.cpp
--- a/Engine/Source/Runtime/AssetRegistry/Private/AssetDataGatherer.cpp
+++ b/Engine/Source/Runtime/AssetRegistry/Private/AssetDataGatherer.cpp
@@ -4182,6 +4182,7 @@ void FAssetDataGatherer::WaitForIdle()
 			return;
 		}
 	}
+	TRACE_CPUPROFILER_EVENT_SCOPE_STR("FAssetDataGatherer::WaitForIdle");
 	Discovery->WaitForIdle();
 	CHECK_IS_NOT_LOCKED_CURRENT_THREAD(ResultsLock);
 
diff --git a/Engine/Source/Runtime/AssetRegistry/Private/AssetRegistry.cpp b/Engine/Source/Runtime/AssetRegistry/Private/AssetRegistry.cpp
--- a/Engine/Source/Runtime/AssetRegistry/Private/AssetRegistry.cpp
+++ b/Engine/Source/Runtime/AssetRegistry/Private/AssetRegistry.cpp
@@ -1662,6 +1662,7 @@ bool FAssetRegistryImpl::ShouldSortReferencers() const
 
 void UAssetRegistryImpl::SearchAllAssets(bool bSynchronousSearch)
 {
+	TRACE_CPUPROFILER_EVENT_SCOPE_STR("UAssetRegistryImpl::SearchAllAssets");
 	using namespace UE::AssetRegistry::Impl;
 
 	double StartTime = FPlatformTime::Seconds();
diff --git a/Engine/Source/Runtime/Launch/Private/LaunchEngineLoop.cpp b/Engine/Source/Runtime/Launch/Private/LaunchEngineLoop.cpp
--- a/Engine/Source/Runtime/Launch/Private/LaunchEngineLoop.cpp
+++ b/Engine/Source/Runtime/Launch/Private/LaunchEngineLoop.cpp
@@ -4146,6 +4146,7 @@ int32 FEngineLoop::PreInitPostStartupScreen(const TCHAR* CmdLine)
 			FCoreDelegates::OnCommandletPreMain.Broadcast();
 			int32 ErrorLevel;
 			{
+				TRACE_CPUPROFILER_EVENT_SCOPE_TEXT(*WriteToString<512>(TEXT("Commandlet Main "), Commandlet->GetFName()));
 				FTrackedActivityScope CommandletActivity(FTrackedActivity::GetEngineActivity(), *FString::Printf(TEXT("Running %s"), *Commandlet->GetName()), false, FTrackedActivity::ELight::Green);
 				ErrorLevel = Commandlet->Main(CommandletCommandLine);
 			}
