Load boot hotfix cvars after config system initialized, persistency dir checking has dependency to it on console.

#tests Tested on console with boot hotfix applied.

diff --git a/Engine/Source/Runtime/Core/Private/Misc/ConfigUtilities.cpp b/Engine/Source/Runtime/Core/Private/Misc/ConfigUtilities.cpp
--- a/Engine/Source/Runtime/Core/Private/Misc/ConfigUtilities.cpp
+++ b/Engine/Source/Runtime/Core/Private/Misc/ConfigUtilities.cpp
@@ -43,6 +43,7 @@ void LoadCVarsFromFileForNextBoot(TMap<FString, FString>& OutCVars)
 {
 	if (!FPaths::HasProjectPersistentDownloadDir())
 	{
+		UE_LOG(LogConfig, Log, TEXT("No project persistent download dir available for boot hotfix"));
 		return;
 	}
 
diff --git a/Engine/Source/Runtime/Launch/Private/LaunchEngineLoop.cpp b/Engine/Source/Runtime/Launch/Private/LaunchEngineLoop.cpp
--- a/Engine/Source/Runtime/Launch/Private/LaunchEngineLoop.cpp
+++ b/Engine/Source/Runtime/Launch/Private/LaunchEngineLoop.cpp
@@ -2780,8 +2780,6 @@ int32 FEngineLoop::PreInitPreStartupScreen(const TCHAR* CmdLine)
 		UE::ConfigUtilities::RecordApplyCVarSettingsFromIni();
 	}
 
-	UE::ConfigUtilities::ApplyCVarsFromBootHotfix();
-
 #if WITH_ENGINE
 	extern ENGINE_API void InitializeRenderingCVarsCaching();
 	InitializeRenderingCVarsCaching();
@@ -6424,6 +6422,12 @@ bool FEngineLoop::AppInit( )
 	}
 #endif
 
+	{
+		SCOPED_BOOT_TIMING("ConfigUtilities::ApplyCVarsFromBootHotfix");
+		// Apply boot hotfixes
+		UE::ConfigUtilities::ApplyCVarsFromBootHotfix();
+	}
+
 #if USE_IO_DISPATCHER
 	// Initialize on demand I/O dispatcher backend
 	FModuleManager::Get().LoadModule("IoStoreOnDemand");
