- Avoid double calls to FEngineLoop::AppExit() (was happening in a Program that called it while exiting, but Mac also calls AppExit in it's standard Wrapper code)

diff --git a/Engine/Source/Runtime/Launch/Private/LaunchEngineLoop.cpp b/Engine/Source/Runtime/Launch/Private/LaunchEngineLoop.cpp
--- a/Engine/Source/Runtime/Launch/Private/LaunchEngineLoop.cpp
+++ b/Engine/Source/Runtime/Launch/Private/LaunchEngineLoop.cpp
@@ -6838,6 +6838,13 @@ void FEngineLoop::AppPreExit( )
 
 void FEngineLoop::AppExit()
 {
+	static bool bCalledOnce;
+	if (bCalledOnce)
+	{
+		return;
+	}
+	bCalledOnce = true;
+	
 	// when compiled WITH_ENGINE, this will happen in FEngineLoop::Exit()
 #if !WITH_ENGINE
 #if STATS
