Fixed LauchEngineLoop to not include some include files if it isn't being built with CoreUObject

diff --git a/Engine/Source/Runtime/Launch/Private/LaunchEngineLoop.cpp b/Engine/Source/Runtime/Launch/Private/LaunchEngineLoop.cpp
--- a/Engine/Source/Runtime/Launch/Private/LaunchEngineLoop.cpp
+++ b/Engine/Source/Runtime/Launch/Private/LaunchEngineLoop.cpp
@@ -81,7 +81,7 @@
 
 #include "Compression/OodleDataCompression.h"
 
-#if WITH_VERSE_VM || defined(__INTELLISENSE__)
+#if (WITH_VERSE_VM || defined(__INTELLISENSE__)) && WITH_COREUOBJECT
 #include "VerseVM/VVMVerse.h"
 #endif
 
@@ -3454,7 +3454,7 @@ int32 FEngineLoop::PreInitPreStartupScreen(const TCHAR* CmdLine)
 	PreInitContext.CommandletCommandLine = CommandletCommandLine;
 #endif // UE_EDITOR || WITH_ENGINE
 
-#if WITH_VERSE_VM || defined(__INTELLISENSE__)
+#if (WITH_VERSE_VM || defined(__INTELLISENSE__)) && WITH_COREUOBJECT
 	Verse::VerseVM::Startup();
 #endif
 
@@ -6817,7 +6817,7 @@ void FEngineLoop::AppPreExit( )
 		GShaderCompilerStats = nullptr;
 	}
 
-#if WITH_VERSE_VM || defined(__INTELLISENSE__)
+#if (WITH_VERSE_VM || defined(__INTELLISENSE__)) && WITH_COREUOBJECT
 	Verse::VerseVM::Shutdown();
 #endif
 
