World Partition Editor Module: remove early loading of the module that was required by some object construction order that it no more needed.

diff --git a/Engine/Source/Runtime/Engine/Private/UnrealEngine.cpp b/Engine/Source/Runtime/Engine/Private/UnrealEngine.cpp
--- a/Engine/Source/Runtime/Engine/Private/UnrealEngine.cpp
+++ b/Engine/Source/Runtime/Engine/Private/UnrealEngine.cpp
@@ -2023,10 +2023,6 @@ void UEngine::Init(IEngineLoop* InEngineLoop)
 	SelectedMaterialColor = DefaultSelectedMaterialColor;
 	SelectionOutlineColor = DefaultSelectedMaterialColor;
 
-#if WITH_EDITOR
-	FModuleManager::Get().LoadModule("WorldPartitionEditor");
-#endif
-
 	InitializeObjectReferences();
 
 	if (GConfig)
diff --git a/Engine/Source/Runtime/Launch/Private/LaunchEngineLoop.cpp b/Engine/Source/Runtime/Launch/Private/LaunchEngineLoop.cpp
--- a/Engine/Source/Runtime/Launch/Private/LaunchEngineLoop.cpp
+++ b/Engine/Source/Runtime/Launch/Private/LaunchEngineLoop.cpp
@@ -4672,10 +4672,6 @@ bool FEngineLoop::LoadStartupCoreModules()
 	// Required during construction of UAnimBlueprint which could happen from async loading thread.
 	// See UAnimBlueprint::UAnimBlueprint().
 	FModuleManager::Get().LoadModule(TEXT("AnimGraph"));
-
-	// Required during construction of UWorldPartitionEditorHash which could happen from async loading thread.
-	// See UWorldPartitionEditorHash::UWorldPartitionEditorHash().
-	FModuleManager::Get().LoadModule(TEXT("WorldPartitionEditor"));
 #endif
 
 	FModuleManager::Get().LoadModule(TEXT("PacketHandler"));
