UE-194878 xr.DisableOpenXROnAndroidWithoutOculus evaluated too late to modify via INI
-Move PreInitHMDDevice to a point quite a bit later in PreInitPreStartupScreen so that more unreal engine systems are setup when we call it, including ini reading, cvars, fplatform memory, etc.

diff --git a/Engine/Source/Runtime/Launch/Private/LaunchEngineLoop.cpp b/Engine/Source/Runtime/Launch/Private/LaunchEngineLoop.cpp
--- a/Engine/Source/Runtime/Launch/Private/LaunchEngineLoop.cpp
+++ b/Engine/Source/Runtime/Launch/Private/LaunchEngineLoop.cpp
@@ -3216,6 +3216,11 @@ int32 FEngineLoop::PreInitPreStartupScreen(const TCHAR* CmdLine)
 		FShaderParametersMetadata::InitializeAllUniformBufferStructs();
 	}
 
+	{
+		SCOPED_BOOT_TIMING("PreInitHMDDevice()");
+		PreInitHMDDevice();
+	}
+
 	{
 		SCOPED_BOOT_TIMING("RHIInit");
 		// Initialize the RHI.
@@ -6656,8 +6661,6 @@ bool FEngineLoop::AppInit( )
 
 	FEmbeddedCommunication::ForceTick(17);
 
-	PreInitHMDDevice();
-
 	// after the above has run we now have the REQUIRED set of engine .INIs  (all of the other .INIs)
 	// that are gotten from .h files' config() are not requires and are dynamically loaded when the .u files are loaded
 
