[Insights] Fixed "Frame" CPU timing events for the Render thread to not be added when the "RenderCommands" trace channel is on. This is because the "BeginFrame" and "EndFrame" commands will define CPU scopes (when "RenderCommands" channel is on) and will overlap with the "Frame" timing event (resulting in incorrect timestamps for both "BeginFrame" and "Frame" timing events). Also renamed "Frame" to "RenderingFrame" for the timer to not be merged with one used in the Game thread.

diff --git a/Engine/Source/Runtime/Launch/Private/LaunchEngineLoop.cpp b/Engine/Source/Runtime/Launch/Private/LaunchEngineLoop.cpp
--- a/Engine/Source/Runtime/Launch/Private/LaunchEngineLoop.cpp
+++ b/Engine/Source/Runtime/Launch/Private/LaunchEngineLoop.cpp
@@ -5456,7 +5456,10 @@ uint64 FScopedSampleMallocChurn::DumpFrame = 0;
 
 #endif
 
+#if CPUPROFILERTRACE_ENABLED
 static uint32 TraceFrameEventThreadId = (uint32) -1;
+static uint32 TraceFrameEventSpecId = 0;
+#endif
 
 static inline void BeginFrameRenderThread(FRHICommandListImmediate& RHICmdList, uint64 CurrentFrameCounter)
 {
@@ -5468,6 +5471,7 @@ static inline void BeginFrameRenderThread(FRHICommandListImmediate& RHICmdList,
 	}
 
 	TRACE_BEGIN_FRAME(TraceFrameType_Rendering);
+
 	GRHICommandList.LatchBypass();
 	GFrameNumberRenderThread++;
 
@@ -5483,10 +5487,14 @@ static inline void BeginFrameRenderThread(FRHICommandListImmediate& RHICmdList,
 
 #if CPUPROFILERTRACE_ENABLED
 	TraceFrameEventThreadId = (uint32) -1;
-	if (UE_TRACE_CHANNELEXPR_IS_ENABLED(CpuChannel))
+	if (UE_TRACE_CHANNELEXPR_IS_ENABLED(CpuChannel) && !UE_TRACE_CHANNELEXPR_IS_ENABLED(RenderCommandsChannel))
 	{
 		TraceFrameEventThreadId = FPlatformTLS::GetCurrentThreadId();
-		FCpuProfilerTrace::OutputBeginDynamicEvent(TEXT("Frame"), __FILE__, __LINE__);
+		if (TraceFrameEventSpecId == 0)
+		{
+			TraceFrameEventSpecId = FCpuProfilerTrace::OutputEventType(TEXT("RenderingFrame"), __FILE__, __LINE__);
+		}
+		FCpuProfilerTrace::OutputBeginEvent(TraceFrameEventSpecId);
 	}
 #endif //CPUPROFILERTRACE_ENABLED
 
