Logging: Added env var UE_LOG_JSON_TO_STDOUT=1 to enable structured logging on stdout for consistency with other tools

diff --git a/Engine/Source/Runtime/Launch/Private/LaunchEngineLoop.cpp b/Engine/Source/Runtime/Launch/Private/LaunchEngineLoop.cpp
--- a/Engine/Source/Runtime/Launch/Private/LaunchEngineLoop.cpp
+++ b/Engine/Source/Runtime/Launch/Private/LaunchEngineLoop.cpp
@@ -385,6 +385,13 @@ public:
 #endif
 
 		bIsJsonOutput = FParse::Param(FCommandLine::Get(), TEXT("JsonStdOut"));
+		if (!bIsJsonOutput)
+		{
+			if (FString Env = FPlatformMisc::GetEnvironmentVariable(TEXT("UE_LOG_JSON_TO_STDOUT")); !Env.IsEmpty())
+			{
+				bIsJsonOutput = FCString::Atoi(*Env) != 0;
+			}
+		}
 
 		if (FParse::Param(FCommandLine::Get(), TEXT("AllowStdOutLogVerbosity")))
 		{
