[shaders] generate brdf headers for the default platform before initializing shader types (the latter now has a side-effect of pre-loading and caching shader source file information for this platform, which means if the brdf headers need to change it will cache an outdated version)

diff --git a/Engine/Source/Runtime/Engine/Private/ShaderCompiler/ShaderCompiler.cpp b/Engine/Source/Runtime/Engine/Private/ShaderCompiler/ShaderCompiler.cpp
--- a/Engine/Source/Runtime/Engine/Private/ShaderCompiler/ShaderCompiler.cpp
+++ b/Engine/Source/Runtime/Engine/Private/ShaderCompiler/ShaderCompiler.cpp
@@ -1122,10 +1122,6 @@ void FShaderJobCache::CullOutputsToMemoryBudget(uint64 TargetBudgetBytes)
 
 
 #if WITH_EDITOR
-static FDelayedAutoRegisterHelper GKickOffShaderAutoGen(EDelayedRegisterRunPhase::ShaderTypesReady, []
-{
-	FShaderCompileUtilities::GenerateBrdfHeaders(GMaxRHIShaderPlatform);
-});
 
 static FDelayedAutoRegisterHelper GKickOffShaderAutoGenForPlatforms(EDelayedRegisterRunPhase::DeviceProfileManagerReady, []
 {
diff --git a/Engine/Source/Runtime/Launch/Private/LaunchEngineLoop.cpp b/Engine/Source/Runtime/Launch/Private/LaunchEngineLoop.cpp
--- a/Engine/Source/Runtime/Launch/Private/LaunchEngineLoop.cpp
+++ b/Engine/Source/Runtime/Launch/Private/LaunchEngineLoop.cpp
@@ -3326,6 +3326,13 @@ int32 FEngineLoop::PreInitPreStartupScreen(const TCHAR* CmdLine)
 		{
 			UE_SCOPED_ENGINE_ACTIVITY(TEXT("Initializing Shader Types"));
 			SCOPED_BOOT_TIMING("InitializeShaderTypes");
+
+#if WITH_EDITOR
+			// Explicitly generate AutogenShaderHeaders.ush prior to shader type initialization
+			// (since that process will load and cache this header as a sideeffect)
+			FShaderCompileUtilities::GenerateBrdfHeaders(GMaxRHIShaderPlatform);
+#endif
+
 			// Initialize shader types before loading any shaders
 			InitializeShaderTypes();
 		}
