Android SingleInstance Service

diff --git a/Engine/Source/Runtime/Core/Private/Android/AndroidPlatformMisc.cpp b/Engine/Source/Runtime/Core/Private/Android/AndroidPlatformMisc.cpp
--- a/Engine/Source/Runtime/Core/Private/Android/AndroidPlatformMisc.cpp
+++ b/Engine/Source/Runtime/Core/Private/Android/AndroidPlatformMisc.cpp
@@ -349,7 +349,7 @@ static struct
 	double	TimeOfChange;
 } CurrentVolume;
 
-#if USE_ANDROID_JNI
+#if USE_ANDROID_JNI && !USE_ANDROID_STANDALONE
 extern "C"
 {
 
@@ -397,7 +397,7 @@ extern "C"
 }
 #endif
 
-#if USE_ANDROID_JNI
+#if USE_ANDROID_JNI && !USE_ANDROID_STANDALONE
 
 // Manage Java side OS event receivers.
 static struct
@@ -416,6 +416,8 @@ static struct
 
 void InitializeJavaEventReceivers()
 {
+	UE_LOG(LogAndroid, Log, TEXT("InitializeJavaEventReceivers"));
+
 	// Register natives to receive Volume, Battery, Headphones events
 	JNIEnv* JEnv = AndroidJavaEnv::GetJavaEnv();
 	if (nullptr != JEnv)
@@ -578,7 +580,7 @@ void FAndroidMisc::PlatformInit()
 	}
 #endif
 
-#if USE_ANDROID_JNI
+#if USE_ANDROID_JNI && !USE_ANDROID_STANDALONE
 	InitializeJavaEventReceivers();
 	AndroidOnBackgroundBinding = FCoreDelegates::ApplicationWillEnterBackgroundDelegate.AddStatic(EnableJavaEventReceivers, false);
 	AndroidOnForegroundBinding = FCoreDelegates::ApplicationHasEnteredForegroundDelegate.AddStatic(EnableJavaEventReceivers, true);
@@ -615,7 +617,7 @@ void FAndroidMisc::PlatformTearDown()
 void FAndroidMisc::UpdateDeviceOrientation()
 {
 	QUICK_SCOPE_CYCLE_COUNTER(STAT_FAndroidMisc_UpdateDeviceOrientation);
-#if USE_ANDROID_JNI
+#if USE_ANDROID_JNI && !USE_ANDROID_STANDALONE
 	JNIEnv* JEnv = AndroidJavaEnv::GetJavaEnv();
 	if (JEnv)
 	{
@@ -641,7 +643,7 @@ void FAndroidMisc::UpdateDeviceOrientation()
 
 void FAndroidMisc::PlatformHandleSplashScreen(bool ShowSplashScreen)
 {
-#if USE_ANDROID_JNI
+#if USE_ANDROID_JNI && !USE_ANDROID_STANDALONE
 	if (!ShowSplashScreen)
 	{
 		AndroidThunkCpp_DismissSplashScreen();
@@ -3036,7 +3038,7 @@ int FAndroidMisc::GetMobilePropagateAlphaSetting()
 TArray<int32> FAndroidMisc::GetSupportedNativeDisplayRefreshRates()
 {
 	TArray<int32> Result;
-#if USE_ANDROID_JNI
+#if USE_ANDROID_JNI && !USE_ANDROID_STANDALONE
 	extern TArray<int32> AndroidThunkCpp_GetSupportedNativeDisplayRefreshRates();
 	Result = AndroidThunkCpp_GetSupportedNativeDisplayRefreshRates();
 #else
@@ -3047,7 +3049,7 @@ TArray<int32> FAndroidMisc::GetSupportedNativeDisplayRefreshRates()
 
 bool FAndroidMisc::SetNativeDisplayRefreshRate(int32 RefreshRate)
 {
-#if USE_ANDROID_JNI
+#if USE_ANDROID_JNI && !USE_ANDROID_STANDALONE
 	extern bool AndroidThunkCpp_SetNativeDisplayRefreshRate(int32 RefreshRate);
 	return AndroidThunkCpp_SetNativeDisplayRefreshRate(RefreshRate);
 #else
@@ -3057,7 +3059,7 @@ bool FAndroidMisc::SetNativeDisplayRefreshRate(int32 RefreshRate)
 
 int32 FAndroidMisc::GetNativeDisplayRefreshRate()
 {
-#if USE_ANDROID_JNI
+#if USE_ANDROID_JNI && !USE_ANDROID_STANDALONE
 	extern int32 AndroidThunkCpp_GetNativeDisplayRefreshRate();
 	return AndroidThunkCpp_GetNativeDisplayRefreshRate();
 #else
@@ -3155,12 +3157,12 @@ void FAndroidMisc::SetAllowedDeviceOrientation(EDeviceScreenOrientation NewAllow
 {
 	AllowedDeviceOrientation = NewAllowedDeviceOrientation;
 
-#if USE_ANDROID_JNI
+#if USE_ANDROID_JNI && !USE_ANDROID_STANDALONE
 	AndroidThunkCpp_SetOrientation(GetAndroidScreenOrientation(NewAllowedDeviceOrientation));
 #endif // USE_ANDROID_JNI
 }
 
-#if USE_ANDROID_JNI
+#if USE_ANDROID_JNI && !USE_ANDROID_STANDALONE
 int32 FAndroidMisc::GetAndroidScreenOrientation(EDeviceScreenOrientation ScreenOrientation)
 {
 	EAndroidScreenOrientation AndroidScreenOrientation = EAndroidScreenOrientation::SCREEN_ORIENTATION_UNSPECIFIED;
@@ -3200,7 +3202,7 @@ int32 FAndroidMisc::GetAndroidScreenOrientation(EDeviceScreenOrientation ScreenO
 
 	return static_cast<int32>(AndroidScreenOrientation);
 }
-#endif // USE_ANDROID_JNI
+#endif // USE_ANDROID_JNI && !USE_ANDROID_STANDALONE
 
 extern void AndroidThunkCpp_ShowConsoleWindow();
 void FAndroidMisc::ShowConsoleWindow()
diff --git a/Engine/Source/Runtime/Core/Private/Android/AndroidPlatformProcess.cpp b/Engine/Source/Runtime/Core/Private/Android/AndroidPlatformProcess.cpp
--- a/Engine/Source/Runtime/Core/Private/Android/AndroidPlatformProcess.cpp
+++ b/Engine/Source/Runtime/Core/Private/Android/AndroidPlatformProcess.cpp
@@ -157,7 +157,7 @@ FString FAndroidPlatformProcess::GetGameBundleId()
 	JNIEnv* JEnv = AndroidJavaEnv::GetJavaEnv();
 	if (nullptr != JEnv)
 	{
-		jclass Class = AndroidJavaEnv::FindJavaClassGlobalRef("com/epicgames/unreal/GameActivity");
+		jclass Class = AndroidJavaEnv::FindJavaClassGlobalRef(ANDROID_GAMEACTIVITY_CLASSPATH);
 		if (nullptr != Class)
 		{
 			jmethodID getAppPackageNameMethodId = JEnv->GetStaticMethodID(Class, "getAppPackageName", "()Ljava/lang/String;");
diff --git a/Engine/Source/Runtime/Core/Private/Android/AndroidPlatformThermal.cpp b/Engine/Source/Runtime/Core/Private/Android/AndroidPlatformThermal.cpp
--- a/Engine/Source/Runtime/Core/Private/Android/AndroidPlatformThermal.cpp
+++ b/Engine/Source/Runtime/Core/Private/Android/AndroidPlatformThermal.cpp
@@ -26,7 +26,7 @@
 #include <android/thermal.h>
 #include "Android/AndroidJava.h"
 
-#if USE_ANDROID_JNI
+#if USE_ANDROID_JNI && !USE_ANDROID_STANDALONE
 namespace UE
 {
 	namespace Android
diff --git a/Engine/Source/Runtime/Core/Public/Android/AndroidPlatform.h b/Engine/Source/Runtime/Core/Public/Android/AndroidPlatform.h
--- a/Engine/Source/Runtime/Core/Public/Android/AndroidPlatform.h
+++ b/Engine/Source/Runtime/Core/Public/Android/AndroidPlatform.h
@@ -101,6 +101,30 @@ typedef FAndroidTypes FPlatformTypes;
 #ifndef USE_ANDROID_EVENTS
 	#define USE_ANDROID_EVENTS						1
 #endif
+#ifndef USE_ANDROID_STANDALONE
+	#define USE_ANDROID_STANDALONE					0
+#endif
+
+
+#if (!USE_ANDROID_STANDALONE && !(UE_BUILD_SHIPPING || UE_BUILD_TEST)) || (USE_ANDROID_STANDALONE && !UE_BUILD_SHIPPING)
+	// M is the scope for the logging such as LogAndroid, STANDALONE_DEBUG_LOGf should be used when using formatted arguments.
+#	define STANDALONE_DEBUG_LOG(M, ...)  FPlatformMisc::LowLevelOutputDebugStringf(TEXT(#M " : "),##__VA_ARGS__);
+#	define STANDALONE_DEBUG_LOGf(M, ...)  FPlatformMisc::LowLevelOutputDebugStringf(TEXT(#M " : ") __VA_ARGS__);
+
+#else
+#	define STANDALONE_DEBUG_LOG(...)  
+#	define STANDALONE_DEBUG_LOGf(...)  
+#endif
+
+#define ANDROID_GAMEACTIVITY_BASE_CLASSPATH			"com/epicgames/unreal/GameActivity"
+
+#if USE_ANDROID_STANDALONE
+#	define ANDROID_GAMEACTIVITY_CLASSPATH	"com/epicgames/makeaar/GameActivityForMakeAAR"
+#else
+#	define ANDROID_GAMEACTIVITY_CLASSPATH	"com/epicgames/unreal/GameActivity"
+#endif
+
+
 
 // Enable to set thread nice values when setting runnable thread priorities
 #ifndef ANDROID_USE_NICE_VALUE_THREADPRIORITY
diff --git a/Engine/Source/Runtime/CoreUObject/Private/Serialization/AsyncLoading.cpp b/Engine/Source/Runtime/CoreUObject/Private/Serialization/AsyncLoading.cpp
--- a/Engine/Source/Runtime/CoreUObject/Private/Serialization/AsyncLoading.cpp
+++ b/Engine/Source/Runtime/CoreUObject/Private/Serialization/AsyncLoading.cpp
@@ -8194,8 +8194,11 @@ bool FAsyncArchive::PrecacheInternal(int64 RequestOffset, int64 RequestSize, boo
 			if (ReadRequestPtr->PollCompletion())
 			{
 				CompleteRead();
-				check(RequestOffset >= PrecacheStartPos && RequestOffset + RequestSize <= PrecacheEndPos);
-				bResult = true;
+				if (GetLoadError() == ELoadError::Unknown)
+				{
+					check(RequestOffset >= PrecacheStartPos && RequestOffset + RequestSize <= PrecacheEndPos);
+					bResult = true;
+				}
 			}
 			delete Read;
 			return bResult;
diff --git a/Engine/Source/Runtime/Launch/Private/Android/AndroidEventManager.cpp b/Engine/Source/Runtime/Launch/Private/Android/AndroidEventManager.cpp
--- a/Engine/Source/Runtime/Launch/Private/Android/AndroidEventManager.cpp
+++ b/Engine/Source/Runtime/Launch/Private/Android/AndroidEventManager.cpp
@@ -73,7 +73,7 @@ void FAppEventManager::Tick()
 	while (!Queue.IsEmpty())
 	{
 		FAppEventPacket Event = DequeueAppEvent();
-		FPlatformMisc::LowLevelOutputDebugStringf(TEXT("FAppEventManager::Tick processing, %d"), int(Event.State));
+		STANDALONE_DEBUG_LOGf(LogAndroidEvents, TEXT("FAppEventManager::Tick processing, %d"), int(Event.State));
 
 		switch (Event.State)
 		{
@@ -97,7 +97,7 @@ void FAppEventManager::Tick()
 			break;
 		case APP_EVENT_STATE_WINDOW_DESTROYED:
 			bHaveWindow = false;
-			FPlatformMisc::LowLevelOutputDebugStringf(TEXT("APP_EVENT_STATE_WINDOW_DESTROYED, %d, %d, %d"), int(bRunning), int(bHaveWindow), int(bHaveGame));
+			STANDALONE_DEBUG_LOGf(LogAndroidEvents, TEXT("APP_EVENT_STATE_WINDOW_DESTROYED, %d, %d, %d"), int(bRunning), int(bHaveWindow), int(bHaveGame));
 			break;
 		case APP_EVENT_STATE_ON_START:
 			//doing nothing here
@@ -105,7 +105,7 @@ void FAppEventManager::Tick()
 		case APP_EVENT_STATE_ON_DESTROY:
 			check(bHaveWindow == false);
 			check(IsEngineExitRequested()); //destroy immediately. Game will shutdown.
-			FPlatformMisc::LowLevelOutputDebugStringf(TEXT("APP_EVENT_STATE_ON_DESTROY"));
+			STANDALONE_DEBUG_LOGf(LogAndroidEvents, TEXT("APP_EVENT_STATE_ON_DESTROY"));
 			break;
 		case APP_EVENT_STATE_ON_STOP:
 			bHaveGame = false;
@@ -134,11 +134,11 @@ void FAppEventManager::Tick()
 		}
 		case APP_EVENT_STATE_APP_ACTIVATED:
 			bRunning = true;
-			FPlatformMisc::LowLevelOutputDebugStringf(TEXT("Execution will be resumed!"));
+			STANDALONE_DEBUG_LOGf(LogAndroidEvents, TEXT("Execution will be resumed!"));
 			break;
 		case APP_EVENT_STATE_APP_SUSPENDED:
 			bRunning = false;
-			FPlatformMisc::LowLevelOutputDebugStringf(TEXT("Execution will be paused..."));
+			STANDALONE_DEBUG_LOGf(LogAndroidEvents, TEXT("Execution will be paused..."));
 			break;
 		case APP_EVENT_STATE_SAFE_ZONE_UPDATED:
 			FCoreDelegates::OnSafeFrameChangedEvent.Broadcast();
@@ -155,7 +155,7 @@ void FAppEventManager::Tick()
 				ExecWindowCreated();
 				bCreateWindow = false;
 				bHaveWindow = true;
-				FPlatformMisc::LowLevelOutputDebugStringf(TEXT("ExecWindowCreated, %d, %d, %d"), int(bRunning), int(bHaveWindow), int(bHaveGame));
+				STANDALONE_DEBUG_LOGf(LogAndroidEvents, TEXT("ExecWindowCreated, %d, %d, %d"), int(bRunning), int(bHaveWindow), int(bHaveGame));
 			}
 		}
 	}
@@ -167,9 +167,9 @@ void FAppEventManager::Tick()
 
 	if (!bRunning)
 	{
-		FPlatformMisc::LowLevelOutputDebugStringf(TEXT("FAppEventManager::Tick EventHandlerEvent Wait "));
+		STANDALONE_DEBUG_LOGf(LogAndroidEvents, TEXT("FAppEventManager::Tick EventHandlerEvent Wait "));
 		EventHandlerEvent->Wait();
-		FPlatformMisc::LowLevelOutputDebugStringf(TEXT("FAppEventManager::Tick EventHandlerEvent DONE Wait "));
+		STANDALONE_DEBUG_LOGf(LogAndroidEvents, TEXT("FAppEventManager::Tick EventHandlerEvent DONE Wait "));
 	}
 }
 
@@ -177,7 +177,7 @@ void FAppEventManager::ReleaseMicrophone(bool shuttingDown)
 {
 	if (FModuleManager::Get().IsModuleLoaded("Voice"))
 	{
-		UE_LOG(LogTemp, Log, TEXT("Android release microphone"));
+		UE_LOG(LogAndroidEvents, Log, TEXT("Android release microphone"));
 		FModuleManager::Get().UnloadModule("Voice", shuttingDown);
 	}
 }
@@ -221,19 +221,34 @@ void FAppEventManager::OnScaleFactorChanged(IConsoleVariable* CVar)
 {
 	if ((CVar->GetFlags() & ECVF_SetByMask) == ECVF_SetByConsole)
 	{
+		STANDALONE_DEBUG_LOGf(LogAndroidEvents, TEXT("FAppEventManager::OnScaleFactorChanged with flags=0x%x"), CVar->GetFlags());
 		FAppEventManager::GetInstance()->ExecWindowResized();
 	}
+	else
+    {
+		STANDALONE_DEBUG_LOGf(LogAndroidEvents, TEXT("NOOP FAppEventManager::OnScaleFactorChanged with flags=0x%x"), CVar->GetFlags());
+    }
+}
+
+static bool IsAppPausedOrSuspended()
+{
+	return FAppEventManager::GetInstance()->IsGamePaused() && !FAppEventManager::GetInstance()->IsGameInFocus();
 }
 
 void FAppEventManager::HandleWindowCreated_EventThread(void* InWindow)
 {
+	if (InWindow == nullptr)
+	{
+		STANDALONE_DEBUG_LOGf(LogAndroidEvents, TEXT("FAppEventManager::HandleWindowCreated_EventThread was given a NULL window so IGNORE."));
+		return;
+	}
+
 	bool AlreadyInited = FirstInitialized;
 
 	// Make sure window will not be deleted until event is processed
 	// Window could be deleted by OS while event queue stuck at game start-up phase
 	FAndroidWindow::AcquireWindowRef((ANativeWindow*)InWindow);
 
-	check(FAndroidWindow::GetHardwareWindow_EventThread() == NULL);
 	FAndroidWindow::SetHardwareWindow_EventThread(InWindow);
 
 	if (!AlreadyInited)
@@ -242,16 +257,29 @@ void FAppEventManager::HandleWindowCreated_EventThread(void* InWindow)
 		FirstInitialized = true;
 	}
 	EnqueueAppEvent(APP_EVENT_STATE_WINDOW_CREATED, FAppEventData((ANativeWindow*)InWindow));
+
+#if USE_ANDROID_STANDALONE
+	FApp::SetHasFocusFunction(&IsAppPausedOrSuspended);
+#endif
 }
 
 void FAppEventManager::HandleWindowClosed_EventThread()
 {
-	check(FAndroidWindow::GetHardwareWindow_EventThread());
+	void* ActiveWindow = FAndroidWindow::GetHardwareWindow_EventThread();
+	
+	if(ActiveWindow == nullptr)
+	{
+		return;
+	}
 
-	FAndroidWindow::ReleaseWindowRef((ANativeWindow*)FAndroidWindow::GetHardwareWindow_EventThread());
 	FAndroidWindow::SetHardwareWindow_EventThread(nullptr);
+	FAndroidWindow::ReleaseWindowRef((ANativeWindow*)ActiveWindow);
 
 	EnqueueAppEvent(APP_EVENT_STATE_WINDOW_DESTROYED);
+
+#if USE_ANDROID_STANDALONE
+	FApp::SetHasFocusFunction(nullptr);
+#endif
 }
 
 
@@ -325,12 +353,12 @@ void FAppEventManager::PauseAudio()
 {
 	if (!GEngine || !GEngine->IsInitialized())
 	{
-		UE_LOG(LogTemp, Log, TEXT("Engine not initialized, not pausing Android audio"));
+		UE_LOG(LogAndroidEvents, Log, TEXT("Engine not initialized, not pausing Android audio"));
 		return;
 	}
 
 	bAudioPaused = true;
-	UE_LOG(LogTemp, Log, TEXT("Android pause audio"));
+	UE_LOG(LogAndroidEvents, Log, TEXT("Android pause audio"));
 
 	FAudioDeviceHandle AudioDevice = GEngine->GetMainAudioDevice();
 	if (AudioDevice)
@@ -348,12 +376,12 @@ void FAppEventManager::ResumeAudio()
 {
 	if (!GEngine || !GEngine->IsInitialized())
 	{
-		UE_LOG(LogTemp, Log, TEXT("Engine not initialized, not resuming Android audio"));
+		UE_LOG(LogAndroidEvents, Log, TEXT("Engine not initialized, not resuming Android audio"));
 		return;
 	}
 
 	bAudioPaused = false;
-	UE_LOG(LogTemp, Log, TEXT("Android resume audio"));
+	UE_LOG(LogAndroidEvents, Log, TEXT("Android resume audio"));
 
 	FAudioDeviceHandle AudioDevice = GEngine->GetMainAudioDevice();
 	if (AudioDevice)
@@ -381,7 +409,7 @@ void FAppEventManager::EnqueueAppEvent(EAppEventState InState, FAppEventData&& I
 	rc = pthread_mutex_unlock(&QueueMutex);
 	check(rc == 0);
 
-	FPlatformMisc::LowLevelOutputDebugStringf(TEXT("LogAndroidEvents::EnqueueAppEvent : %u, [width=%d, height=%d], tid = %d, %s"), InState, InData.WindowWidth, InData.WindowHeight, gettid(), GetAppEventName(InState));
+	STANDALONE_DEBUG_LOGf(LogAndroidEvents, TEXT("LogAndroidEvents::EnqueueAppEvent : %u, [width=%d, height=%d], tid = %d, %s"), InState, InData.WindowWidth, InData.WindowHeight, gettid(), GetAppEventName(InState));
 }
 
 FAppEventPacket FAppEventManager::DequeueAppEvent()
@@ -409,7 +437,11 @@ bool FAppEventManager::IsGamePaused()
 
 bool FAppEventManager::IsGameInFocus()
 {
+#if USE_ANDROID_STANDALONE
+	return (bWindowInFocus && bHaveWindow && bHaveGame);
+#else
 	return (bWindowInFocus && bHaveWindow);
+#endif
 }
 
 
@@ -463,6 +495,4 @@ bool FAppEventManager::WaitForEventInQueue(EAppEventState InState, double Timeou
 	return FoundEvent;
 }
 
-extern volatile bool GEventHandlerInitialized;
-
 #endif
diff --git a/Engine/Source/Runtime/Launch/Private/Android/AndroidJNI.cpp b/Engine/Source/Runtime/Launch/Private/Android/AndroidJNI.cpp
--- a/Engine/Source/Runtime/Launch/Private/Android/AndroidJNI.cpp
+++ b/Engine/Source/Runtime/Launch/Private/Android/AndroidJNI.cpp
@@ -74,7 +74,16 @@ if (Id == 0) \
 void FJavaWrapper::FindClassesAndMethods(JNIEnv* Env)
 {
 	auto bIsOptional = false;
-	GGameActivityClassID = GameActivityClassID = FindClassGlobalRef(Env, "com/epicgames/unreal/GameActivity", bIsOptional);
+	FString classPath = ANDROID_GAMEACTIVITY_CLASSPATH;
+	FString classPathBase = ANDROID_GAMEACTIVITY_BASE_CLASSPATH;
+	STANDALONE_DEBUG_LOGf(LogAndroid, TEXT("[JNI] - FindClassesAndMethods called for: %s"), *classPath);
+	jclass foundGameActivityClassID = FindClassGlobalRef(Env, TCHAR_TO_ANSI(*classPath), bIsOptional);
+	if (!foundGameActivityClassID)
+	{
+		STANDALONE_DEBUG_LOGf(LogAndroid, TEXT("[JNI] - FindClassesAndMethods FAILED to find GameActivity ClassID for: %s"), *classPath);
+		return;
+	}
+	GGameActivityClassID = GameActivityClassID = foundGameActivityClassID;
 	AndroidThunkJava_ShowConsoleWindow = FindMethod(Env, GameActivityClassID, "AndroidThunkJava_ShowConsoleWindow", "(Ljava/lang/String;)V", bIsOptional);
     AndroidThunkJava_ShowVirtualKeyboardInputDialog = FindMethod(Env, GameActivityClassID, "AndroidThunkJava_ShowVirtualKeyboardInputDialog", "(ILjava/lang/String;Ljava/lang/String;)V", bIsOptional);
     AndroidThunkJava_HideVirtualKeyboardInputDialog = FindMethod(Env, GameActivityClassID, "AndroidThunkJava_HideVirtualKeyboardInputDialog", "()V", bIsOptional);
@@ -1297,6 +1306,11 @@ jobject AndroidJNI_GetJavaAssetManager()
 		{
 			auto local = NewScopedJavaObject(Env, FJavaWrapper::CallObjectMethod(Env, FJavaWrapper::GameActivityThis, FJavaWrapper::AndroidThunkJava_GetAssetManager));
 			GJavaAssetManager = (jobject)Env->NewGlobalRef(*local);
+			checkf(GJavaAssetManager != nullptr, TEXT("Failed get GJavaAssetManager!"));
+		}
+		else
+		{
+			checkf(FAndroidApplication::GetJavaEnv() != nullptr, TEXT("Failed get FAndroidApplication::GetJavaEnv() "));
 		}
 	}
 	return GJavaAssetManager;
@@ -1707,6 +1721,13 @@ JNI_METHOD void Java_com_epicgames_unreal_GameActivity_nativeSetObbFilePaths(JNI
 //This function is declared in the Java-defined class, GameActivity.java: "public native void nativeSetGlobalActivity();"
 JNI_METHOD void Java_com_epicgames_unreal_GameActivity_nativeSetGlobalActivity(JNIEnv* jenv, jobject thiz, jboolean bUseExternalFilesDir, jboolean bPublicLogFiles, jstring internalFilePath, jstring externalFilePath, jboolean bOBBinAPK, jstring APKFilename /*, jobject googleServices*/)
 {
+	STANDALONE_DEBUG_LOG(TEXT("nativeSetGlobalActivity(unreal): Entering unreal nativeSetGlobalActivity, GameActivityThis=%p\n"), FJavaWrapper::GameActivityThis);
+	if (FJavaWrapper::GameActivityThis != nullptr)
+	{
+		STANDALONE_DEBUG_LOG(TEXT("nativeSetGlobalActivity(unreal): Error GameActivityThis is already set GameActivityThis=%p\n"), FJavaWrapper::GameActivityThis);
+		//jenv->DeleteGlobalRef(FJavaWrapper::GameActivityThis);
+		GGameActivityThis = FJavaWrapper::GameActivityThis = nullptr;
+	}
 	if (!FJavaWrapper::GameActivityThis)
 	{
 		GGameActivityThis = FJavaWrapper::GameActivityThis = jenv->NewGlobalRef(thiz);
@@ -1715,6 +1736,7 @@ JNI_METHOD void Java_com_epicgames_unreal_GameActivity_nativeSetGlobalActivity(J
 			FPlatformMisc::LowLevelOutputDebugString(TEXT("Error setting the global GameActivity activity"));
 			check(false);
 		}
+		STANDALONE_DEBUG_LOG(TEXT("nativeSetGlobalActivity(unreal): jenv=%p, set GameActivityThis=%p, bUseExternalFilesDir=%d, bOBBinAPK=%d\n"), jenv, FJavaWrapper::GameActivityThis, bUseExternalFilesDir, bOBBinAPK);
 
 		// This call is only to set the correct GameActivityThis
 		FAndroidApplication::InitializeJavaEnv(GJavaVM, JNI_CURRENT_VERSION, FJavaWrapper::GameActivityThis);
@@ -1740,11 +1762,11 @@ JNI_METHOD void Java_com_epicgames_unreal_GameActivity_nativeSetGlobalActivity(J
 #else
 			GFilePathBase = GExternalFilePath;
 #endif
-			FPlatformMisc::LowLevelOutputDebugStringf(TEXT("GFilePathBase Path override to'%s'\n"), *GFilePathBase);
+			FPlatformMisc::LowLevelOutputDebugStringf(TEXT("nativeSetGlobalActivity(unreal): GFilePathBase Path override to'%s'\n"), *GFilePathBase);
 		}
 
-		FPlatformMisc::LowLevelOutputDebugStringf(TEXT("InternalFilePath found as '%s'\n"), *GInternalFilePath);
-		FPlatformMisc::LowLevelOutputDebugStringf(TEXT("ExternalFilePath found as '%s'\n"), *GExternalFilePath);
+		FPlatformMisc::LowLevelOutputDebugStringf(TEXT("nativeSetGlobalActivity(unreal): InternalFilePath found as '%s'\n"), *GInternalFilePath);
+		FPlatformMisc::LowLevelOutputDebugStringf(TEXT("nativeSetGlobalActivity(unreal): ExternalFilePath found as '%s'\n"), *GExternalFilePath);
 	}
 }
 
@@ -2078,6 +2100,56 @@ JNI_METHOD void Java_com_epicgames_unreal_NativeCalls_UELogVerbose(JNIEnv* jenv,
 #endif
 	jenv->ReleaseStringUTFChars(InString, chars);
 }
+#if USE_ANDROID_STANDALONE
+JNI_METHOD void Java_com_epicgames_makeaar_GameActivityForMakeAAR_nativeSetGlobalActivity(JNIEnv* jenv, jobject thiz, jboolean bUseExternalFilesDir, jboolean bPublicLogFiles, jstring internalFilePath, jstring externalFilePath, jboolean bOBBinAPK, jstring APKFilename /*, jobject googleServices*/)
+{
+	STANDALONE_DEBUG_LOG(TEXT("nativeSetGlobalActivity: Entering makeaar nativeSetGlobalActivity, GameActivityThis=%p\n"), FJavaWrapper::GameActivityThis);
+	if (FJavaWrapper::GameActivityThis != nullptr)
+	{
+		STANDALONE_DEBUG_LOG(TEXT("nativeSetGlobalActivity(makeaar): Error GameActivityThis is already set GameActivityThis=%p\n"), FJavaWrapper::GameActivityThis);
+		//jenv->DeleteGlobalRef(FJavaWrapper::GameActivityThis);
+		GGameActivityThis = FJavaWrapper::GameActivityThis = nullptr;
+	}
+	if (!FJavaWrapper::GameActivityThis)
+	{
+		GGameActivityThis = FJavaWrapper::GameActivityThis = jenv->NewGlobalRef(thiz);
+		if (!FJavaWrapper::GameActivityThis)
+		{
+			FPlatformMisc::LowLevelOutputDebugString(TEXT("Error setting the global GameActivity activity"));
+			check(false);
+		}
+		STANDALONE_DEBUG_LOG(TEXT("nativeSetGlobalActivity(makeaar): jenv=%p, set GameActivityThis=%p, bUseExternalFilesDir=%d, bOBBinAPK=%d\n"), jenv, FJavaWrapper::GameActivityThis, bUseExternalFilesDir, bOBBinAPK);
+		// This call is only to set the correct GameActivityThis
+		FAndroidApplication::InitializeJavaEnv(GJavaVM, JNI_CURRENT_VERSION, FJavaWrapper::GameActivityThis);
+		// Rescan methods since we are switching to the makeaar version of GameActivity
+		FJavaWrapper::FindClassesAndMethods(jenv);
+		// @todo split GooglePlay, this needs to be passed in to this function
+		FJavaWrapper::GoogleServicesThis = FJavaWrapper::GameActivityThis;
+		// FJavaWrapper::GoogleServicesThis = jenv->NewGlobalRef(googleServices);
+		// Next we check to see if the OBB file is in the APK
+		//jmethodID isOBBInAPKMethod = jenv->GetStaticMethodID(FJavaWrapper::GameActivityClassID, "isOBBInAPK", "()Z");
+		//GOBBinAPK = (bool)jenv->CallStaticBooleanMethod(FJavaWrapper::GameActivityClassID, isOBBInAPKMethod, nullptr);
+		GOBBinAPK = bOBBinAPK;
+		GAPKFilename = FJavaHelper::FStringFromParam(jenv, APKFilename);
+		GInternalFilePath = FJavaHelper::FStringFromParam(jenv, internalFilePath);
+		GExternalFilePath = FJavaHelper::FStringFromParam(jenv, externalFilePath);
+		if (bUseExternalFilesDir)
+		{
+#if UE_BUILD_SHIPPING
+			GFilePathBase = GInternalFilePath;
+			GOverrideAndroidLogDir = bPublicLogFiles;
+#else
+			GFilePathBase = GExternalFilePath;
+#endif
+			FPlatformMisc::LowLevelOutputDebugStringf(TEXT("nativeSetGlobalActivity(makeaar): GFilePathBase Path override to'%s'\n"), *GFilePathBase);
+		}
+		FPlatformMisc::LowLevelOutputDebugStringf(TEXT("nativeSetGlobalActivity(makeaar): InternalFilePath found as '%s'\n"), *GInternalFilePath);
+		FPlatformMisc::LowLevelOutputDebugStringf(TEXT("nativeSetGlobalActivity(makeaar): ExternalFilePath found as '%s'\n"), *GExternalFilePath);
+	}
+	// DO NOT CALL, we have already handled logic above...
+	// Java_com_epicgames_unreal_GameActivity_nativeSetGlobalActivity(jenv, thiz, bUseExternalFilesDir, bPublicLogFiles, internalFilePath, externalFilePath, bOBBinAPK, APKFilename);
+}
+#endif // USE_ANDROID_STANDALONE
 
 void FJavaWrapper::SetupEmbeddedCommunication(JNIEnv* Env)
 {
diff --git a/Engine/Source/Runtime/Launch/Private/Android/LaunchAndroid.cpp b/Engine/Source/Runtime/Launch/Private/Android/LaunchAndroid.cpp
--- a/Engine/Source/Runtime/Launch/Private/Android/LaunchAndroid.cpp
+++ b/Engine/Source/Runtime/Launch/Private/Android/LaunchAndroid.cpp
@@ -41,6 +41,16 @@
 #include <jni.h>
 #include <android/sensor.h>
 
+#include "HAL/ExceptionHandling.h"
+#include "ProfilingDebugging/LoadTimeTracker.h"
+#include "Android/AndroidPlatformCrashContext.h"
+
+#if USE_ANDROID_STANDALONE
+#include "HAL/FileManager.h"
+#include "HAL/PlatformFilemanager.h"
+#include "IPlatformFilePak.h"
+#endif
+
 // Function pointer for retrieving joystick events
 // Function has been part of the OS since Honeycomb, but only appeared in the
 // NDK in r19. Querying via dlsym allows its use without tying to the newest
@@ -206,12 +216,20 @@ static TAutoConsoleVariable<int32> CVarEnableSustainedPerformanceMode(
 	TEXT("  1: Enabled"),
 	ECVF_Default);
 
+static TAutoConsoleVariable<int32> CVarUseGameThreadForCommands(
+	TEXT("Android.UseGameThread"),
+	0,
+	TEXT("Use GameThread for incoming commands\n")
+	TEXT("  0: Disabled (default)\n")
+	TEXT("  1: Enabled"),
+	ECVF_Default);
+
 extern void AndroidThunkCpp_SetSustainedPerformanceMode(bool);
 static void SetSustainedPerformanceMode()
 {
 	static bool bSustainedPerformanceMode = false;
 	bool bIncomingSustainedPerformanceMode = CVarEnableSustainedPerformanceMode.GetValueOnAnyThread() != 0;
-	if(bSustainedPerformanceMode != bIncomingSustainedPerformanceMode)
+	if (bSustainedPerformanceMode != bIncomingSustainedPerformanceMode)
 	{
 		bSustainedPerformanceMode = bIncomingSustainedPerformanceMode;
 		UE_LOG(LogAndroid, Log, TEXT("Setting sustained performance mode: %d"), (int32)bSustainedPerformanceMode);
@@ -226,17 +244,151 @@ static FEvent* EventHandlerEvent = NULL;
 // Wait for Java onCreate to complete before resume main init
 static volatile bool GResumeMainInit = false;
 volatile bool GEventHandlerInitialized = false;
+static void IssueConsoleCommand(FString Command)
+{
+	static TArray<FString> GPendingConsoleCommands;
+	if (GEngine != NULL)
+	{
+		// Run on game thread to avoid race condition with DeferredCommands
+		AsyncTask(ENamedThreads::GameThread, [Command]()
+			{
+				if (GPendingConsoleCommands.Num() > 0)
+				{
+					STANDALONE_DEBUG_LOGf(LogAndroid, TEXT("IssueConsoleCommand: Pending count= %d"), GPendingConsoleCommands.Num());
+					for (FString& pendingCommand : GPendingConsoleCommands)
+					{
+						STANDALONE_DEBUG_LOGf(LogAndroid, TEXT("IssuePendingConsoleCommand: %s"), *pendingCommand);
+						GEngine->DeferredCommands.Add(pendingCommand);
+					}
+					GPendingConsoleCommands.Empty();
+				}
+				STANDALONE_DEBUG_LOGf(LogAndroid, TEXT("IssueConsoleCommand: %s"), *Command);
+				GEngine->DeferredCommands.Add(Command);
+			});
+	}
+// to avoid out of order console command issues that were overwritten in configs, for non standalone we do not store pending.
+#if USE_ANDROID_STANDALONE
+	else
+	{
+		STANDALONE_DEBUG_LOGf(LogAndroid, TEXT("Ignoring console command (too early, Queue for later): %s"), *Command);
+		GPendingConsoleCommands.Add_GetRef(MoveTemp(Command));
+	}
+#endif
+}
 
 // The event thread locks this whenever the window unavailable during early init, pause and resume.
 FCriticalSection GAndroidWindowLock;
+static bool bAppIsActive_EventThread = false;
+static bool bDidGainFocus = false;
+static bool bHasFocus = false;
+static bool bHasWindow = false;
+static bool bIsResumed = false;
+static bool bShouldRestartFromInterrupt = false;
+static bool bIgnorePauseOnDownloaderStart = false;
 bool bReadyToProcessEvents = false;
+static int GAndroidWindowLockRefCount = 0;
+static const bool bForce_GAndroidWindowLock = true;
+extern void* GAndroidWindowOverride;
+
+void GAndroidWindowLock_Lock(FString calledBy)
+{
+	bool doAction = bForce_GAndroidWindowLock;
+	if (GAndroidWindowLockRefCount >= 0)
+	{
+		++GAndroidWindowLockRefCount;
+		if (GAndroidWindowLockRefCount == 1)
+		{
+			doAction = true;
+		}
+		else
+		{
+			STANDALONE_DEBUG_LOGf(LogAndroid, TEXT("WARNING GAndroidWindowLock_Lock(%s). ALREADY LOCKED. GAndroidWindowLockRefCount=%d, bAppIsActive_EventThread=%d"), *calledBy, GAndroidWindowLockRefCount, bAppIsActive_EventThread);
+		}
+	}
+	else
+	{
+		STANDALONE_DEBUG_LOGf(LogAndroid, TEXT("GAndroidWindowLock_Lock(%s). ERROR: GAndroidWindowLockRefCount=%d, bAppIsActive_EventThread=%d"), *calledBy, GAndroidWindowLockRefCount, bAppIsActive_EventThread);
+	}
+	if (doAction)
+	{
+		GAndroidWindowLock.Lock();
+
+		STANDALONE_DEBUG_LOGf(LogAndroid, TEXT("GAndroidWindowLock_Lock(%s). LOCKED. GAndroidWindowLockRefCount=%d, bAppIsActive_EventThread=%d"), *calledBy, GAndroidWindowLockRefCount, bAppIsActive_EventThread);
+	}
+	else
+	{
+		STANDALONE_DEBUG_LOGf(LogAndroid, TEXT("WARNING GAndroidWindowLock_Lock(%s). NOOP: GAndroidWindowLockRefCount=%d, bAppIsActive_EventThread=%d bReadyToProcessEvents=%d"), *calledBy, GAndroidWindowLockRefCount, bAppIsActive_EventThread, bReadyToProcessEvents);
+	}
+	//return GAndroidWindowLockRefCount>0;
+}
+
+void GAndroidWindowLock_Unlock(FString calledBy)
+{
+	bool doAction = bForce_GAndroidWindowLock;
+	if (GAndroidWindowLockRefCount > 0)
+	{
+		if (--GAndroidWindowLockRefCount == 0)
+		{
+			doAction = true;
+		}
+		else
+		{
+			STANDALONE_DEBUG_LOGf(LogAndroid, TEXT("WARNING GAndroidWindowLock_Unlock(%s). KEEP LOCKED, GAndroidWindowLockRefCount=%d, bAppIsActive_EventThread=%d"), *calledBy, GAndroidWindowLockRefCount, bAppIsActive_EventThread);
+		}
+		if (GAndroidWindowLockRefCount < 0)
+		{
+			STANDALONE_DEBUG_LOGf(LogAndroid, TEXT("GAndroidWindowLock_Unlock(%s). ERROR UNLOCK MISMATCH, GAndroidWindowLockRefCount=%d, bAppIsActive_EventThread=%d"), *calledBy, GAndroidWindowLockRefCount, bAppIsActive_EventThread);
+			GAndroidWindowLockRefCount = 0;
+		}
+	}
+	else
+	{
+		STANDALONE_DEBUG_LOGf(LogAndroid, TEXT("GAndroidWindowLock_Unlock(%s). SKIPPED, GAndroidWindowLockRefCount=%d, bAppIsActive_EventThread=%d"), *calledBy, GAndroidWindowLockRefCount, bAppIsActive_EventThread);
+	}
+	if (doAction)
+	{
+		GAndroidWindowLock.Unlock();
+		//bAppIsActive_EventThread = true;
+		STANDALONE_DEBUG_LOGf(LogAndroid, TEXT("GAndroidWindowLock_Unlock(%s). UNLOCK, GAndroidWindowLockRefCount=%d, bAppIsActive_EventThread=%d"), *calledBy, GAndroidWindowLockRefCount, bAppIsActive_EventThread);
+	}
+	else
+	{
+		STANDALONE_DEBUG_LOGf(LogAndroid, TEXT("GAndroidWindowLock_Unlock(%s). NOOP: GAndroidWindowLockRefCount=%d, bAppIsActive_EventThread=%d bReadyToProcessEvents=%d"), *calledBy, GAndroidWindowLockRefCount, bAppIsActive_EventThread, bReadyToProcessEvents);
+	}
+	//return GAndroidWindowLockRefCount==0;
+}
+static void SuspendApp_EventThread();
+static void ActivateApp_EventThread();
+#if USE_ANDROID_STANDALONE && !(UE_BUILD_SHIPPING || UE_BUILD_TEST)
+#define DEVELOPER_LOG_COMMANDCB_CASE(a) STANDALONE_DEBUG_LOGf(LogAndroid, TEXT("--TRACE-- OnAppCommandCB - Case [%s] : bDidGainFocus = %d, bHasFocus = %d, bHasWindow = %d, bIsResumed = %d, bShouldRestartFromInterrupt = %d, bReadyToProcessEvents=%d, bAppIsActive_EventThread=%d"), TEXT(#a), bDidGainFocus, bHasFocus, bHasWindow, bIsResumed, bShouldRestartFromInterrupt, bReadyToProcessEvents, bAppIsActive_EventThread)
+#else
+#define DEVELOPER_LOG_COMMANDCB_CASE(a) 
+#endif
+
 void FPlatformMisc::UnlockAndroidWindow()
 {
 	check(IsInGameThread());
 	check(FTaskGraphInterface::IsRunning());
+
+	DEVELOPER_LOG_COMMANDCB_CASE(UnlockAndroidWindow);
 	UE_LOG(LogAndroid, Log, TEXT("Unlocking android HW window during preinit."));
+	
 	bReadyToProcessEvents = true;
-	GAndroidWindowLock.Unlock();
+	GAndroidWindowLock_Unlock("UnlockAndroidWindow");
+	
+}
+
+void LockAndroidWindow(FString calledBy)
+{
+	check(IsInGameThread());
+	check(FTaskGraphInterface::IsRunning());
+	
+	DEVELOPER_LOG_COMMANDCB_CASE(LockAndroidWindow);
+	UE_LOG(LogAndroid, Log, TEXT("Locking android HW window during preinit. bAppIsActive_EventThread=%d"), bAppIsActive_EventThread);
+	
+	GAndroidWindowLock_Lock(*(calledBy + "->LockAndroidWindow"));
+	bReadyToProcessEvents = false;
+	//bAppIsActive_EventThread = false;
 }
 
 JNI_METHOD void Java_com_epicgames_unreal_GameActivity_nativeResumeMainInit(JNIEnv* jenv, jobject thiz)
@@ -274,11 +426,26 @@ void InitHMDs()
 	}
 }
 
-extern AAssetManager * AndroidThunkCpp_GetAssetManager();
+extern AAssetManager* AndroidThunkCpp_GetAssetManager();
+#if USE_ANDROID_STANDALONE
+	FString GAndroidCommandLine;
+#endif
+/**
+ * Shuts down the engine
+ */
+void EngineExit(void)
+{
+	// Make sure this is set
+	RequestEngineExit(TEXT("EngineExit() was called"));
+	GEngineLoop.Exit();
+	UE_LOG(LogAndroid, Log, TEXT("Exiting is over"));
+	FPlatformMisc::RequestExit(1);
+}
 
 static void InitCommandLine()
 {
 	static const uint32 CMD_LINE_MAX = 16384u;
+	STANDALONE_DEBUG_LOGf(LogAndroid, TEXT("LaunchAndroid::InitCommandLine called!"));
 
 	// initialize the command line to an empty string
 	FCommandLine::Set(TEXT(""));
@@ -318,6 +485,15 @@ static void InitCommandLine()
 			FCommandLine::Append(UTF8_TO_TCHAR(CommandLine));
 			FPlatformMisc::LowLevelOutputDebugStringf(TEXT("APK Commandline: %s"), FCommandLine::Get());
 		}
+#if !UE_BUILD_SHIPPING && USE_ANDROID_STANDALONE
+		UE_SET_LOG_VERBOSITY(LogAndroid, Log);
+		if (!GAndroidCommandLine.IsEmpty())
+		{
+			FCommandLine::Set(TEXT(""));
+			FCommandLine::Append(*GAndroidCommandLine);
+			STANDALONE_DEBUG_LOGf(LogAndroid, TEXT("Override Commandline: %s"), FCommandLine::Get());
+		}
+#endif
 
 		// read in the command line text file from the sdcard if it exists
 		FString CommandLineFilePath = GFilePathBase + FString("/UnrealGame/") + (!FApp::IsProjectNameEmpty() ? FApp::GetProjectName() : FPlatformProcess::ExecutableName()) + FString("/UECommandLine.txt");
@@ -386,7 +562,7 @@ extern void AndroidThunkCpp_DismissSplashScreen();
 static void OnNativeWindowResized(ANativeActivity* activity, ANativeWindow* window)
 {
 	static int8_t cmd = APP_CMD_WINDOW_RESIZED;
-	struct android_app* app = (struct android_app *)activity->instance;
+	struct android_app* app = (struct android_app*)activity->instance;
 	write(app->msgwrite, &cmd, sizeof(cmd));
 }
 
@@ -403,7 +579,7 @@ static void ApplyAndroidCompatConfigRules()
 			static const TCHAR AndroidCompat[] = TEXT("AndroidCompat.");
 			if (Key.StartsWith(AndroidCompat))
 			{
-				FString CVarName = Key.Mid(UE_ARRAY_COUNT(AndroidCompat)-1);
+				FString CVarName = Key.Mid(UE_ARRAY_COUNT(AndroidCompat) - 1);
 				if (AllowedCompatCVars.Contains(CVarName))
 				{
 					auto* CVar = IConsoleManager::Get().FindConsoleVariable(*CVarName);
@@ -429,14 +605,20 @@ int32 AndroidMain(struct android_app* state)
 
 	// Force the first call to GetJavaEnv() to happen on the game thread, allowing subsequent calls to occur on any thread
 	FAndroidApplication::GetJavaEnv();
-
+#if !USE_ANDROID_STANDALONE
 	// Set window format to 8888
 	ANativeActivity_setWindowFormat(state->activity, WINDOW_FORMAT_RGBA_8888);
+#endif
+
+#if !UE_BUILD_SHIPPING && USE_ANDROID_STANDALONE
+	STANDALONE_DEBUG_LOGf(LogAndroid, TEXT("ENABLE LogAndroid verbosity!\n"));
+	UE_SET_LOG_VERBOSITY(LogAndroid, Log);
+#endif
 
 	// adjust the file descriptor limits to allow as many open files as possible
 	rlimit cur_fd_limit;
 	{
-		int result = getrlimit(RLIMIT_NOFILE, & cur_fd_limit);
+		int result = getrlimit(RLIMIT_NOFILE, &cur_fd_limit);
 		//FPlatformMisc::LowLevelOutputDebugStringf(TEXT("(%d) Current fd limits: soft = %lld, hard = %lld"), result, cur_fd_limit.rlim_cur, cur_fd_limit.rlim_max);
 	}
 	{
@@ -447,7 +629,7 @@ int32 AndroidMain(struct android_app* state)
 		//FPlatformMisc::LowLevelOutputDebugStringf(TEXT("(%d) Setting fd limits: soft = %lld, hard = %lld"), result, new_limit.rlim_cur, new_limit.rlim_max);
 	}
 	{
-		int result = getrlimit(RLIMIT_NOFILE, & cur_fd_limit);
+		int result = getrlimit(RLIMIT_NOFILE, &cur_fd_limit);
 		//FPlatformMisc::LowLevelOutputDebugStringf(TEXT("(%d) Current fd limits: soft = %lld, hard = %lld"), result, cur_fd_limit.rlim_cur, cur_fd_limit.rlim_max);
 	}
 
@@ -456,7 +638,7 @@ int32 AndroidMain(struct android_app* state)
 	// However, it has existed in the so since Honeycomb, query for the symbol
 	// to determine whether to try controller support
 	{
-		void* Lib = dlopen("libandroid.so",0);
+		void* Lib = dlopen("libandroid.so", 0);
 		if (Lib != NULL)
 		{
 			GetAxes = (GetAxesType)dlsym(Lib, "AMotionEvent_getAxisValue");
@@ -488,13 +670,14 @@ int32 AndroidMain(struct android_app* state)
 		AlwaysAllowedKeyCodes.Add(AlwaysAllowedKeyCodesList[i]);
 	}
 
-	const int ValidGamepadKeyCodeCount = sizeof(ValidGamepadKeyCodesList)/sizeof(uint32);
+	const int ValidGamepadKeyCodeCount = sizeof(ValidGamepadKeyCodesList) / sizeof(uint32);
 	for (int i = 0; i < ValidGamepadKeyCodeCount; ++i)
 	{
 		ValidGamepadKeyCodes.Add(ValidGamepadKeyCodesList[i]);
 	}
 
 	FAndroidPlatformStackWalk::InitStackWalking();
+	STANDALONE_DEBUG_LOGf(LogAndroid, TEXT(" wait for java activity onCreate to finish, GResumeMainInit=%d"), GResumeMainInit);
 
 	// wait for java activity onCreate to finish
 	{
@@ -505,10 +688,23 @@ int32 AndroidMain(struct android_app* state)
 			FPlatformMisc::MemoryBarrier();
 		}
 	}
+	// make sure GEngineLoop::Exit() is always called.
+	struct EngineLoopCleanupGuard
+	{
+		~EngineLoopCleanupGuard()
+		{
+			if (GIsGuarded)
+			{
+				GIsGuarded = false;
+			}
+			EngineExit();
+		}
+	} CleanupGuard;
 
 	// read the command line file
 	InitCommandLine();
 	FPlatformMisc::LowLevelOutputDebugStringf(TEXT("Final commandline: %s (len %i)\n"), FCommandLine::Get(), FCString::Strlen(FCommandLine::Get()));
+	const TCHAR* CmdLine = FCommandLine::Get();
 
 #if !(UE_BUILD_SHIPPING)
 	// If "-waitforattach" or "-WaitForDebugger" was specified, halt startup and wait for a debugger to attach before continuing
@@ -520,6 +716,42 @@ int32 AndroidMain(struct android_app* state)
 	}
 #endif
 
+	STANDALONE_DEBUG_LOGf(LogAndroid, TEXT("LogLoad::AndroidMain - NDKVersion = %d"), PLATFORM_USED_NDK_VERSION_INTEGER );
+
+	bool bEnableGuarded = false;
+
+	FString reportcrashOverride;
+	if (FParse::Value(CmdLine, TEXT("reportcrashOverride="), reportcrashOverride))
+	{
+		GAlwaysReportCrash = FCString::Atoi(*reportcrashOverride);
+		STANDALONE_DEBUG_LOGf(LogAndroid, TEXT("LogLoad::AndroidMain cmdline arg(reportcrashOverride) DISABLE GAlwaysReportCrash. isDebuggerPresent = %d, GAlwaysReportCrash = %d, bEnableGuarded = %d, GIsGuarded = %d"), FPlatformMisc::IsDebuggerPresent(), GAlwaysReportCrash, bEnableGuarded, GIsGuarded);
+	}
+
+
+#if (UE_BUILD_DEBUG || UE_BUILD_DEVELOPMENT)
+	if (true && !GAlwaysReportCrash)
+#else
+	if (FPlatformMisc::IsDebuggerPresent() && !GAlwaysReportCrash)
+#endif
+	{
+		// Don't use exception handling when a debugger is attached to exactly trap the crash. This does NOT check
+		// whether we are the first instance or not!
+		bEnableGuarded = false;
+	}
+	else
+	{
+		bEnableGuarded = true;
+	}
+
+	if (bEnableGuarded && !GIsGuarded)
+	{
+		STANDALONE_DEBUG_LOGf(LogAndroid, TEXT("LogLoad::AndroidMain ENABLE GIsGuarded. isDebuggerPresent = %d, GAlwaysReportCrash = %d, bEnableGuarded = %d, GIsGuarded = %d"), FPlatformMisc::IsDebuggerPresent(), GAlwaysReportCrash, bEnableGuarded, GIsGuarded);
+
+		GIsGuarded = true;
+	}
+
+	STANDALONE_DEBUG_LOGf(LogAndroid, TEXT("LogLoad::AndroidMain (Guarded state) IsDebuggerPresent = %d, GAlwaysReportCrash = %d, bEnableGuarded = %d, GIsGuarded = %d"), FPlatformMisc::IsDebuggerPresent(), GAlwaysReportCrash, bEnableGuarded, GIsGuarded);
+
 	EventHandlerEvent = FPlatformProcess::GetSynchEventFromPool(false);
 	FPlatformMisc::LowLevelOutputDebugString(TEXT("Created sync event\n"));
 	FAppEventManager::GetInstance()->SetEventHandlerEvent(EventHandlerEvent);
@@ -540,9 +772,11 @@ int32 AndroidMain(struct android_app* state)
 		// Note: this is unlocked after Android's PlatformCreateDynamicRHI when the RHI is then able to process window changes.
 		// We don't wait for all of preinit to complete as PreLoadScreens will need to process events during preinit.
 
-		UE_LOG(LogAndroid, Log, TEXT("PreInit android HW window lock."));
-		GAndroidWindowLock.Lock();
+		UE_LOG(LogAndroid, Log, TEXT("PreInit android HW window lock. bAppIsActive_EventThread=%d"), bAppIsActive_EventThread);
+		DEVELOPER_LOG_COMMANDCB_CASE(PreInitWindowLock);
+		GAndroidWindowLock_Lock("AndroidMain lock");
 	}
+	FPlatformMisc::LowLevelOutputDebugString( TEXT("After GAndroidWindowLock in AndroidMain"));
 
 	FDelegateHandle ConfigReadyHandle = FCoreDelegates::TSConfigReadyForUse().AddStatic(&ApplyAndroidCompatConfigRules);
 
@@ -557,11 +791,14 @@ int32 AndroidMain(struct android_app* state)
 		return PreInitResult;
 	}
 
+#if !USE_ANDROID_STANDALONE
+
 	// register callback for native window resize
 	if (GAndroidEnableNativeResizeEvent)
 	{
 		state->activity->callbacks->onNativeWindowResized = OnNativeWindowResized;
 	}
+#endif
 
 	// initialize HMDs
 	{
@@ -572,6 +809,7 @@ int32 AndroidMain(struct android_app* state)
 	UE_LOG(LogAndroid, Display, TEXT("Passed PreInit()"));
 
 	GLog->SetCurrentThreadAsPrimaryThread();
+	int32 ErrorLevel = PreInitResult;
 
 	FAppEventManager::GetInstance()->SetEmptyQueueHandlerEvent(FPlatformProcess::GetSynchEventFromPool(false));
 
@@ -580,6 +818,9 @@ int32 AndroidMain(struct android_app* state)
 	bDidCompleteEngineInit = true;
 
 	UE_LOG(LogAndroid, Log, TEXT("Passed GEngineLoop.Init()"));
+	double EngineInitializationTime = FPlatformTime::Seconds() - GStartTime;
+	UE_LOG(LogLoad, Log, TEXT("(Engine Initialization) Total time: %.2f seconds"), EngineInitializationTime);
+	ACCUM_LOADTIME(TEXT("EngineInitialization"), EngineInitializationTime);
 
 	AndroidThunkCpp_DismissSplashScreen();
 
@@ -600,13 +841,22 @@ int32 AndroidMain(struct android_app* state)
 
 	BootTimingPoint("Tick loop starting");
 	DumpBootTiming();
+	if (bEnableGuarded && !GIsGuarded)
+	{
+		UE_LOG(LogLoad, Log, TEXT("(AndroidMain) ENABLE GIsGuarded before Engine Tick, IsDebuggerPresent = %d, GAlwaysReportCrash = %d, bEnableGuarded = %d, GIsGuarded = %d"), FPlatformMisc::IsDebuggerPresent(), GAlwaysReportCrash, bEnableGuarded, GIsGuarded);
+		GIsGuarded = true;
+	}
 	// tick until done
 	while (!IsEngineExitRequested())
 	{
 		FAndroidStats::UpdateAndroidStats();
 
 		FAppEventManager::GetInstance()->Tick();
-		if(!FAppEventManager::GetInstance()->IsGamePaused())
+		if (!FAppEventManager::GetInstance()->IsGamePaused()
+#if USE_ANDROID_STANDALONE
+			&& FAppEventManager::GetInstance()->IsGameInFocus()
+#endif
+			)
 		{
 			GEngineLoop.Tick();
 		}
@@ -625,6 +875,7 @@ int32 AndroidMain(struct android_app* state)
 		}
 #endif
 	}
+	TRACE_BOOKMARK(TEXT("Tick loop end"));
 	
 	FAppEventManager::GetInstance()->TriggerEmptyQueue();
 
@@ -634,10 +885,18 @@ int32 AndroidMain(struct android_app* state)
 	GEngineLoop.Exit();
 
 	UE_LOG(LogAndroid, Log, TEXT("Exiting is over"));
+#if !USE_ANDROID_STANDALONE
 
 	FPlatformMisc::RequestExit(true, TEXT("AndroidMain"));
+#endif
 	return 0;
 }
+void* AndroidMain(void* param)
+{
+	struct android_app* state = (struct android_app*)param;
+	AndroidMain(state);
+	return nullptr;
+}
 
 struct AChoreographer;
 struct FChoreographer
@@ -788,7 +1047,8 @@ static void* AndroidEventThreadWorker( void* param )
 
 	// window is initially invalid/locked.
 	UE_LOG(LogAndroid, Log, TEXT("event thread, Initial HW window lock."));
-	GAndroidWindowLock.Lock();
+	DEVELOPER_LOG_COMMANDCB_CASE(AndroidEventThreadWorker_BeforeWhile);
+	GAndroidWindowLock_Lock("AndroidEventThreadWorker");
 
 	//continue to process events until the engine is shutting down
 	while (!IsEngineExitRequested())
@@ -799,10 +1059,12 @@ static void* AndroidEventThreadWorker( void* param )
 
 		sleep(EventRefreshRate);		// this is really 0 since it takes int seconds.
 	}
+	DEVELOPER_LOG_COMMANDCB_CASE(AndroidEventThreadWorker_AfterWhile);
 
-	GAndroidWindowLock.Unlock();
+	GAndroidWindowLock_Unlock("AndroidEventThreadWorker");
 
-	UE_LOG(LogAndroid, Log, TEXT("Exiting"));
+	UE_LOG(LogAndroid, Log, TEXT("AndroidEventThreadWorker->Exiting"));
+	STANDALONE_DEBUG_LOGf(LogAndroid, TEXT("Exiting AndroidEventThreadWorker"));
 
 	return NULL;
 }
@@ -815,7 +1077,7 @@ static void AndroidProcessEvents(struct android_app* state)
 	int events;
 	struct android_poll_source* source;
 
-	while((ident = ALooper_pollAll(-1, &fdesc, &events, (void**)&source)) >= 0)
+	while ((ident = ALooper_pollAll(-1, &fdesc, &events, (void**)&source)) >= 0)
 	{
 		// process this event
 		if (source)
@@ -836,10 +1098,13 @@ void android_main(struct android_app* state)
 	GGameThreadId = FPlatformTLS::GetCurrentThreadId();
 
 	BootTimingPoint("android_main");
-	FPlatformMisc::LowLevelOutputDebugString(TEXT("Entering native app glue main function"));
+	STANDALONE_DEBUG_LOGf(LogAndroid, TEXT("Entering native app glue main function"));
 	
 	GNativeAndroidApp = state;
 	check(GNativeAndroidApp);
+#if USE_ANDROID_STANDALONE
+	STANDALONE_DEBUG_LOGf(LogAndroid, TEXT("android_main : USE_ANDROID_STANDALONE is defined and enabled"));
+#endif
 
 	pthread_attr_t otherAttr; 
 	pthread_attr_init(&otherAttr);
@@ -856,6 +1121,12 @@ void android_main(struct android_app* state)
 	//@todo android: replace with native activity, main loop off of UI thread, etc.
 	AndroidMain(state);
 }
+void* android_main(void* param)
+{
+	struct android_app* state = (struct android_app*)param;
+	android_main(state);
+	return nullptr;
+}
 
 extern bool GAndroidGPUInfoReady;
 
@@ -876,7 +1147,7 @@ static bool TryIgnoreClick(AInputEvent* event, size_t actionPointer)
 //Called from the event process thread
 static int32_t HandleInputCB(struct android_app* app, AInputEvent* event)
 {
-//	FPlatformMisc::LowLevelOutputDebugStringf(TEXT("INPUT - type: %x, action: %x, source: %x, keycode: %x, buttons: %x"), AInputEvent_getType(event), 
+//	STANDALONE_DEBUG_LOGf(LogAndroid, TEXT("INPUT - type: %x, action: %x, source: %x, keycode: %x, buttons: %x"), AInputEvent_getType(event), 
 //		AMotionEvent_getAction(event), AInputEvent_getSource(event), AKeyEvent_getKeyCode(event), AMotionEvent_getButtonState(event));
 	check(IsInAndroidEventThread());
 
@@ -1000,12 +1271,12 @@ static int32_t HandleInputCB(struct android_app* app, AInputEvent* event)
 			(GetAxes != NULL) &&
 			(actionType == AMOTION_EVENT_ACTION_MOVE))
 		{
-			const int axisCount = sizeof(AxisList)/sizeof(int32_t);
+			const int axisCount = sizeof(AxisList) / sizeof(int32_t);
 
 			// poll all the axes and forward to update controller state
 			for (int axis = 0; axis < axisCount; axis++)
 			{
-				float val = GetAxes( event, AxisList[axis], 0);
+				float val = GetAxes(event, AxisList[axis], 0);
 				FAndroidInputInterface::JoystickAxisEvent(device, AxisList[axis], val);
 			}
 
@@ -1170,7 +1441,7 @@ static int32_t HandleInputCB(struct android_app* app, AInputEvent* event)
 		int keyFlags = AKeyEvent_getFlags(event);
 		bool bSoftKey = (keyFlags & AKEY_EVENT_FLAG_SOFT_KEYBOARD) != 0;
 
-		FPlatformMisc::LowLevelOutputDebugStringf(TEXT("Received keycode: %d, softkey: %d"), keyCode, bSoftKey ? 1 : 0);
+		STANDALONE_DEBUG_LOGf(LogAndroid, TEXT("Received keycode: %d, softkey: %d"), keyCode, bSoftKey ? 1 : 0);
 
 		//Only pass on the device id if really a gamepad, joystick or dpad (allows menu and back to be treated as gamepad events)
 		int32 device = -1;
@@ -1186,11 +1457,11 @@ static int32_t HandleInputCB(struct android_app* app, AInputEvent* event)
 		{
 			bool down = AKeyEvent_getAction(event) != AKEY_EVENT_ACTION_UP;
 			FAndroidInputInterface::JoystickButtonEvent(device, keyCode, down);
-			FPlatformMisc::LowLevelOutputDebugStringf(TEXT("Received gamepad button: %d"), keyCode);
+			STANDALONE_DEBUG_LOGf(LogAndroid, TEXT("Received gamepad button: %d"), keyCode);
 		}
 		else
 		{
-			FPlatformMisc::LowLevelOutputDebugStringf(TEXT("Received key event: %d"), keyCode);
+			STANDALONE_DEBUG_LOGf(LogAndroid, TEXT("Received key event: %d"), keyCode);
 
 			// only handle mapped key codes
 			if (!MappedKeyCodes.Contains(keyCode))
@@ -1240,8 +1511,6 @@ static int32_t HandleInputCB(struct android_app* app, AInputEvent* event)
 	return 0;
 }
 
-static bool bShouldRestartFromInterrupt = false;
-static bool bIgnorePauseOnDownloaderStart = false;
 static bool IsStartupMoviePlaying()
 {
 	return GEngine && GEngine->IsInitialized() && GetMoviePlayer() && GetMoviePlayer()->IsStartupMoviePlaying();
@@ -1255,31 +1524,67 @@ static bool IsPreLoadScreenPlaying()
 
 FAppEventData::FAppEventData(ANativeWindow* WindowIn)
 {
-	check(WindowIn);
+	if (WindowIn == nullptr)
+	{
+		UE_LOG(LogAndroid, Log, TEXT("FAppEventData was passed a NULL WindowsIn pointer."));
+		return;
+	}
+	//check(WindowIn);
+	extern void* GAndroidWindowOverride;
+	extern int GSurfaceViewWidth;
+	extern int GSurfaceViewHeight;
+	if (GAndroidWindowOverride != nullptr)
+	{
+		if (GSurfaceViewWidth>0 && GSurfaceViewHeight>0)
+		{
+			WindowWidth = GSurfaceViewWidth;	//ANativeWindow_getWidth((ANativeWindow*)GAndroidWindowOverride);
+			WindowHeight = GSurfaceViewHeight;	//ANativeWindow_getHeight((ANativeWindow*)GAndroidWindowOverride);
+		}
+        else
+        {
+			GSurfaceViewWidth = WindowWidth = ANativeWindow_getWidth((ANativeWindow*)GAndroidWindowOverride);
+			GSurfaceViewHeight = WindowHeight = ANativeWindow_getHeight((ANativeWindow*)GAndroidWindowOverride);
+		}
+	}
+	else
+	{
 	WindowWidth = ANativeWindow_getWidth(WindowIn);
 	WindowHeight = ANativeWindow_getHeight(WindowIn);
-	check(WindowWidth >= 0 && WindowHeight >= 0);
+	}
+	if (WindowWidth <= 0)
+	{
+		UE_LOG(LogAndroid, Log, TEXT("FAppEventData was passed a window with ZERO Height."));
+		WindowWidth = 1;
+	}
+	if (WindowHeight <= 0)
+	{
+		UE_LOG(LogAndroid, Log, TEXT("FAppEventData was passed a window with ZERO Height."));
+		WindowHeight = 1;
+	}
+	//check(WindowWidth >= 0 && WindowHeight >= 0);
 }
 
-static bool bAppIsActive_EventThread = false;
 
 
 // called when the app has window + resume.
 static void ActivateApp_EventThread()
 {
+	DEVELOPER_LOG_COMMANDCB_CASE(ActivateApp_EventThread);
 	if (bAppIsActive_EventThread)
 	{
 		// Seems this can occur.
+		STANDALONE_DEBUG_LOGf(LogAndroid, TEXT("WARNING -- event thread, activate app, ALREADY have HW window lock according to bAppIsActive_EventThread. bReadyToProcessEvents=%d, bAppIsActive_EventThread=%d"), bReadyToProcessEvents, bAppIsActive_EventThread);
 		return;
 	}
 
 	// Unlock window when we're ready.
 	UE_LOG(LogAndroid, Log, TEXT("event thread, activate app, unlocking HW window"));
-	GAndroidWindowLock.Unlock();
+	GAndroidWindowLock_Unlock("ActivateApp_EventThread");
 	// wake the GT up.
 	FAppEventManager::GetInstance()->EnqueueAppEvent(APP_EVENT_STATE_APP_ACTIVATED);
 
 	bAppIsActive_EventThread = true;
+	STANDALONE_DEBUG_LOGf(LogAndroid, TEXT("event thread, activate app, UNLOCKED HW window lock. bReadyToProcessEvents=%d, bAppIsActive_EventThread=%d"), bReadyToProcessEvents, bAppIsActive_EventThread);
 	FAppEventManager::GetInstance()->EnqueueAppEvent(APP_EVENT_RUN_CALLBACK, FAppEventData([]()
 	{
 		UE_LOG(LogAndroid, Log, TEXT("performing app foregrounding callback."));
@@ -1295,26 +1600,31 @@ static void ActivateApp_EventThread()
 	}
 
 	FThreadHeartBeat::Get().ResumeHeartBeat(true);
+#if !USE_ANDROID_STANDALONE
 
 	FPreLoadScreenManager::EnableRendering(true);
 
 	extern void AndroidThunkCpp_ShowHiddenAlertDialog();
 	AndroidThunkCpp_ShowHiddenAlertDialog();	
+#endif
 }
 
 extern void BlockRendering();
 // called whenever the app loses loses window or pause.
 static void SuspendApp_EventThread()
 {
+	DEVELOPER_LOG_COMMANDCB_CASE(SuspendApp_EventThread);
 	if (!bAppIsActive_EventThread)
 	{
+		STANDALONE_DEBUG_LOGf(LogAndroid, TEXT("event thread, SuspendApp_EventThread called but aborted, GAndroidWindowLockRefCount=%d, bReadyToProcessEvents=%d, bAppIsActive_EventThread=%d, IsInGameThread=%d"), GAndroidWindowLockRefCount, bReadyToProcessEvents, bAppIsActive_EventThread, IsInGameThread());
 		return;
 	}
 	bAppIsActive_EventThread = false;
 	// Lock the window, this prevents event thread from removing the window whilst the RHI initializes.
 
 	UE_LOG(LogAndroid, Log, TEXT("event thread, suspending app, acquiring HW window lock."));
-	GAndroidWindowLock.Lock();
+	GAndroidWindowLock_Lock("SuspendApp_EventThread");
+	STANDALONE_DEBUG_LOGf(LogAndroid, TEXT("event thread, suspending app, acquiring HW window lock. bReadyToProcessEvents=%d, bAppIsActive_EventThread=%d, IsInGameThread=%d"), bReadyToProcessEvents, bAppIsActive_EventThread, IsInGameThread());
 
 	if (bReadyToProcessEvents == false)
 	{
@@ -1323,7 +1633,9 @@ static void SuspendApp_EventThread()
 		// we ignore events until this point as acquiring GAndroidWindowLock means requires the window to be properly initialized.
 		UE_LOG(LogAndroid, Log, TEXT("event thread, app not yet ready."));
 		return;
-	};
+	}
+
+	STANDALONE_DEBUG_LOGf(LogAndroid, TEXT("event thread, suspending app, ACQUIRED HW window lock. bReadyToProcessEvents=%d, bAppIsActive_EventThread=%d"), bReadyToProcessEvents, bAppIsActive_EventThread);
 
 	TSharedPtr<FEvent, ESPMode::ThreadSafe> EMDoneTrigger = MakeShareable(FPlatformProcess::GetSynchEventFromPool(), [](FEvent* EventToDelete)
 	{
@@ -1343,16 +1655,23 @@ static void SuspendApp_EventThread()
 		EMDoneTrigger->Trigger();
 	}));
 
+	uint32 StartCycles = FPlatformTime::Cycles();
+
 	FEmbeddedCommunication::WakeGameThread();
 
+#if !USE_ANDROID_STANDALONE
 	FPreLoadScreenManager::EnableRendering(false);
+#endif
 
 	FThreadHeartBeat::Get().SuspendHeartBeat(true);
 
 	// wait for a period of time before blocking rendering
 	UE_LOG(LogAndroid, Log, TEXT("AndroidEGL::  SuspendApp_EventThread, waiting for event manager to process. tid: %d"), FPlatformTLS::GetCurrentThreadId());
+	EMDoneTrigger->Reset();
 	bool bSuccess = EMDoneTrigger->Wait(4000);
+	float ElapsedTimeInMs_EMDoneTrigger_Wait = FPlatformTime::ToMilliseconds(FPlatformTime::Cycles() - StartCycles);
 	UE_CLOG(!bSuccess, LogAndroid, Log, TEXT("backgrounding callback, not responded in timely manner."));
+	STANDALONE_DEBUG_LOGf(LogAndroid, TEXT("AndroidEGL::  SuspendApp_EventThread, EMDoneTrigger->Wait, waited '%f' ms"), (float)ElapsedTimeInMs_EMDoneTrigger_Wait);
 
 	BlockRendering();
 
@@ -1364,12 +1683,7 @@ static void SuspendApp_EventThread()
 static void OnAppCommandCB(struct android_app* app, int32_t cmd)
 {
 	check(IsInAndroidEventThread());
-	static bool bDidGainFocus = false;
-	//FPlatformMisc::LowLevelOutputDebugStringf(TEXT("OnAppCommandCB cmd: %u, tid = %d"), cmd, gettid());
 
-	static bool bHasFocus = false;
-	static bool bHasWindow = false;
-	static bool bIsResumed = false;
 
 	// Set event thread's view of the window dimensions:
 	{
@@ -1391,6 +1705,7 @@ static void OnAppCommandCB(struct android_app* app, int32_t cmd)
 		* later.
 		*/
 		// the OS asked us to save the state of the app
+		DEVELOPER_LOG_COMMANDCB_CASE(APP_CMD_SAVE_STATE);
 		UE_LOG(LogAndroid, Log, TEXT("Case APP_CMD_SAVE_STATE"));
 		FAppEventManager::GetInstance()->EnqueueAppEvent(APP_EVENT_STATE_SAVE_STATE);
 		break;
@@ -1401,9 +1716,15 @@ static void OnAppCommandCB(struct android_app* app, int32_t cmd)
 		 * surface.
 		 */
 		// get the window ready for showing
-		FPlatformMisc::LowLevelOutputDebugStringf(TEXT("Case APP_CMD_INIT_WINDOW"));
+		DEVELOPER_LOG_COMMANDCB_CASE(APP_CMD_INIT_WINDOW);
+		STANDALONE_DEBUG_LOGf(LogAndroid, TEXT("Case APP_CMD_INIT_WINDOW"));
 		UE_LOG(LogAndroid, Log, TEXT("Case APP_CMD_INIT_WINDOW"));
+#if USE_ANDROID_STANDALONE
+		check(app->pendingWindow == (ANativeWindow*)GAndroidWindowOverride);
+		FAppEventManager::GetInstance()->HandleWindowCreated_EventThread((ANativeWindow*)GAndroidWindowOverride);
+#else
 		FAppEventManager::GetInstance()->HandleWindowCreated_EventThread(app->pendingWindow);
+#endif
 		bHasWindow = true;
 		if (bHasWindow && bHasFocus && bIsResumed)
 		{
@@ -1417,7 +1738,8 @@ static void OnAppCommandCB(struct android_app* app, int32_t cmd)
 		 * contains the existing window; after calling android_app_exec_cmd
 		 * it will be set to NULL.
 		 */
-		FPlatformMisc::LowLevelOutputDebugStringf(TEXT("Case APP_CMD_TERM_WINDOW, tid = %d"), gettid());
+		DEVELOPER_LOG_COMMANDCB_CASE(APP_CMD_TERM_WINDOW);
+		STANDALONE_DEBUG_LOGf(LogAndroid, TEXT("Case APP_CMD_TERM_WINDOW, tid = %d"), gettid());
 		// clean up the window because it is being hidden/closed
 		UE_LOG(LogAndroid, Log, TEXT("Case APP_CMD_TERM_WINDOW"));
 
@@ -1433,6 +1755,7 @@ static void OnAppCommandCB(struct android_app* app, int32_t cmd)
 		 */
 		// if the app lost focus, avoid unnecessary processing (like monitoring the accelerometer)
 		// log it, but the actual event will be simulated later in APP_CMD_PAUSE
+		DEVELOPER_LOG_COMMANDCB_CASE(APP_CMD_LOST_FOCUS);
 		UE_LOG(LogAndroid, Log, TEXT("Case APP_CMD_LOST_FOCUS"));
 
 		break;
@@ -1441,6 +1764,7 @@ static void OnAppCommandCB(struct android_app* app, int32_t cmd)
 		 * Command from main thread: the app's activity window has gained
 		 * input focus.
 		 */
+		DEVELOPER_LOG_COMMANDCB_CASE(APP_CMD_GAINED_FOCUS);
 
 		// bring back a certain functionality, like monitoring the accelerometer
 		// log it, but the actual event will be simulated later in APP_CMD_RESUME
@@ -1460,6 +1784,7 @@ static void OnAppCommandCB(struct android_app* app, int32_t cmd)
 		 * Command from main thread: the current ANativeWindow has been resized.
 		 * Please redraw with its new size.
 		 */
+		DEVELOPER_LOG_COMMANDCB_CASE(APP_CMD_WINDOW_RESIZED);
 		UE_LOG(LogAndroid, Log, TEXT("Case APP_CMD_WINDOW_RESIZED"));
 		FAppEventManager::GetInstance()->EnqueueAppEvent(APP_EVENT_STATE_WINDOW_RESIZED, FAppEventData(app->window));
 		break;
@@ -1469,8 +1794,9 @@ static void OnAppCommandCB(struct android_app* app, int32_t cmd)
 		 * be redrawn.  You should redraw the window before handing this to
 		 * android_app_exec_cmd() in order to avoid transient drawing glitches.
 		 */
+		DEVELOPER_LOG_COMMANDCB_CASE(APP_CMD_WINDOW_REDRAW_NEEDED);
 		UE_LOG(LogAndroid, Log, TEXT("Case APP_CMD_WINDOW_REDRAW_NEEDED"));
-		FAppEventManager::GetInstance()->EnqueueAppEvent(APP_EVENT_STATE_WINDOW_REDRAW_NEEDED );
+		FAppEventManager::GetInstance()->EnqueueAppEvent(APP_EVENT_STATE_WINDOW_REDRAW_NEEDED);
 		break;
 	case APP_CMD_CONTENT_RECT_CHANGED:
 		/**
@@ -1478,6 +1804,7 @@ static void OnAppCommandCB(struct android_app* app, int32_t cmd)
 		 * such as from the soft input window being shown or hidden.  You can
 		 * find the new content rect in android_app::contentRect.
 		 */
+		DEVELOPER_LOG_COMMANDCB_CASE(APP_CMD_CONTENT_RECT_CHANGED);
 		UE_LOG(LogAndroid, Log, TEXT("Case APP_CMD_CONTENT_RECT_CHANGED"));
 		break;
 	/* receive this event from Java instead to work around NDK bug with AConfiguration_getOrientation in Oreo
@@ -1499,12 +1826,14 @@ static void OnAppCommandCB(struct android_app* app, int32_t cmd)
 		 * Command from main thread: the system is running low on memory.
 		 * Try to reduce your memory use.
 		 */
+		DEVELOPER_LOG_COMMANDCB_CASE(APP_CMD_LOW_MEMORY);
 		UE_LOG(LogAndroid, Log, TEXT("Case APP_CMD_LOW_MEMORY"));
 		break;
 	case APP_CMD_START:
 		/**
 		 * Command from main thread: the app's activity has been started.
 		 */
+		DEVELOPER_LOG_COMMANDCB_CASE(APP_CMD_START);
 		UE_LOG(LogAndroid, Log, TEXT("Case APP_CMD_START"));
 		FAppEventManager::GetInstance()->EnqueueAppEvent(APP_EVENT_STATE_ON_START);
 	
@@ -1518,12 +1847,13 @@ static void OnAppCommandCB(struct android_app* app, int32_t cmd)
 		// assume focus on resume
 		bDidGainFocus = true;
 		bHasFocus = true;
+		DEVELOPER_LOG_COMMANDCB_CASE(APP_CMD_RESUME);
 
 		if (bHasWindow && bHasFocus && bIsResumed)
 		{
 			ActivateApp_EventThread();
 		}
-		FPlatformMisc::LowLevelOutputDebugStringf(TEXT("Case APP_CMD_RESUME"));
+		STANDALONE_DEBUG_LOGf(LogAndroid, TEXT("Case APP_CMD_RESUME, bShouldRestartFromInterrupt=%d"), bShouldRestartFromInterrupt);
 		UE_LOG(LogAndroid, Log, TEXT("Case APP_CMD_RESUME"));
 		FAppEventManager::GetInstance()->EnqueueAppEvent(APP_EVENT_STATE_ON_RESUME);
 
@@ -1544,7 +1874,8 @@ static void OnAppCommandCB(struct android_app* app, int32_t cmd)
 		/**
 		 * Command from main thread: the app's activity has been paused.
 		 */
-		FPlatformMisc::LowLevelOutputDebugStringf(TEXT("Case APP_CMD_PAUSE"));
+		DEVELOPER_LOG_COMMANDCB_CASE(APP_CMD_PAUSE);
+		STANDALONE_DEBUG_LOGf(LogAndroid, TEXT("Case APP_CMD_PAUSE"));
 		UE_LOG(LogAndroid, Log, TEXT("Case APP_CMD_PAUSE"));
 
 		// simulate lost focus
@@ -1601,6 +1932,7 @@ static void OnAppCommandCB(struct android_app* app, int32_t cmd)
 		/**
 		 * Command from main thread: the app's activity has been stopped.
 		 */
+		DEVELOPER_LOG_COMMANDCB_CASE(APP_CMD_STOP);
 		UE_LOG(LogAndroid, Log, TEXT("Case APP_CMD_STOP"));
 		FAppEventManager::GetInstance()->EnqueueAppEvent(APP_EVENT_STATE_ON_STOP);
 		break;
@@ -1609,6 +1941,7 @@ static void OnAppCommandCB(struct android_app* app, int32_t cmd)
 		* Command from main thread: the app's activity is being destroyed,
 		* and waiting for the app thread to clean up and exit before proceeding.
 		*/
+		DEVELOPER_LOG_COMMANDCB_CASE(APP_CMD_DESTROY);
 		UE_LOG(LogAndroid, Log, TEXT("Case APP_CMD_DESTROY"));
 
 		FAppEventManager::GetInstance()->EnqueueAppEvent(APP_EVENT_RUN_CALLBACK, FAppEventData([]()
@@ -1627,10 +1960,14 @@ static void OnAppCommandCB(struct android_app* app, int32_t cmd)
 
 
 		FAppEventManager::GetInstance()->EnqueueAppEvent(APP_EVENT_STATE_ON_DESTROY);
+#if !USE_ANDROID_STANDALONE
 
 		// Exit here, avoids having to unlock the window and letting the RHI's deal with invalid window.
 		extern void AndroidThunkCpp_ForceQuit();
 		AndroidThunkCpp_ForceQuit();
+#else
+		STANDALONE_DEBUG_LOGf(LogAndroid, TEXT("IGNORING ForceQuit in handling APP_CMD_DESTROY"));
+#endif
 
 		break;
 	}
@@ -1687,6 +2024,7 @@ JNI_METHOD void Java_com_epicgames_unreal_GameActivity_nativeOnOrientationChange
 JNI_METHOD void Java_com_epicgames_unreal_GameActivity_nativeConsoleCommand(JNIEnv* jenv, jobject thiz, jstring commandString)
 {
 	FString Command = FJavaHelper::FStringFromParam(jenv, commandString);
+#if !USE_ANDROID_STANDALONE
 	if (GEngine != NULL)
 	{
 		// Run on game thread to avoid race condition with DeferredCommands
@@ -1699,6 +2037,10 @@ JNI_METHOD void Java_com_epicgames_unreal_GameActivity_nativeConsoleCommand(JNIE
 	{
 		FPlatformMisc::LowLevelOutputDebugStringf(TEXT("Ignoring console command (too early): %s"), *Command);
 	}
+#else
+	IssueConsoleCommand(Command);
+#endif
+
 }
 
 // This is called from the Java UI thread for initializing VR HMDs
@@ -1783,4 +2125,293 @@ bool WaitForAndroidLoseFocusEvent(double TimeoutSeconds)
 	return FAppEventManager::GetInstance()->WaitForEventInQueue(EAppEventState::APP_EVENT_STATE_WINDOW_LOST_FOCUS, TimeoutSeconds);
 }
 
+
+#if USE_ANDROID_STANDALONE
+pthread_t G_AndroidMainThread;
+
+static struct android_app dummy_state;
+
+#define EVENTTYPE_INIT						0
+#define EVENTTYPE_POST_ENGINE_INIT			1
+#define EVENTTYPE_ENGINELOOP_INIT_COMPLETE	2
+#define EVENTTYPE_FRAME_BEGIN				3
+#define EVENTTYPE_FRAME_END					4
+#define EVENTTYPE_PRE_LOAD_MAP				5
+#define EVENTTYPE_POST_LOAD_MAP				6
+#define EVENTTYPE_ACTION					7
+
+void AndroidThunkCpp_Engine_SendEvent(int32 event, const FString& param1, int32 param2, int32 param3, float param4)
+{
+	if (JNIEnv* Env = FAndroidApplication::GetJavaEnv())
+	{
+		static jclass EngineClass = FAndroidApplication::FindJavaClassGlobalRef("com/epicgames/makeaar/Engine");
+		static jmethodID SendEventFunc = FJavaWrapper::FindStaticMethod(Env, EngineClass, "AndroidThunkJava_Engine_ReceiveEvent", "(ILjava/lang/String;IIF)V", false);
+
+		if (SendEventFunc != nullptr)
+		{
+			auto StringParam = FJavaHelper::ToJavaString(Env, param1);
+			Env->CallStaticVoidMethod(EngineClass, SendEventFunc, event, *StringParam, param2, param3, param4);
+		}
+		else
+		{
+			STANDALONE_DEBUG_LOGf(LogAndroid, TEXT("AndroidThunkCpp_Engine_SendEvent WARNING: SendEventFunc not found"));
+
+		}
+	}
+	else
+	{
+
+		STANDALONE_DEBUG_LOGf(LogAndroid, TEXT("AndroidThunkCpp_Engine_SendEvent ERROR: FAndroidApplication::GetJavaEnv() == null"));
+	}
+}
+
+static void InitEvent()
+{
+	STANDALONE_DEBUG_LOGf(LogAndroid, TEXT("LaunchAndroid received InitEvent"));
+	AndroidThunkCpp_Engine_SendEvent(EVENTTYPE_INIT, TEXT(""), 0, 0, 0.0f);
+}
+
+static void PostEngineInitEvent()
+{
+	STANDALONE_DEBUG_LOGf(LogAndroid, TEXT("LaunchAndroid received PostEngineInitEvent"));
+
+	AndroidThunkCpp_Engine_SendEvent(EVENTTYPE_POST_ENGINE_INIT, TEXT(""), 0, 0, 0.0f);
+}
+
+static void EngineLoopInitCompleteEvent()
+{
+	STANDALONE_DEBUG_LOGf(LogAndroid, TEXT("LaunchAndroid received EngineLoopInitCompleteEvent"));
+	AndroidThunkCpp_Engine_SendEvent(EVENTTYPE_ENGINELOOP_INIT_COMPLETE, TEXT(""), 0, 0, 0.0f);
+}
+
+static void BeginFrameEvent()
+{
+	AndroidThunkCpp_Engine_SendEvent(EVENTTYPE_FRAME_BEGIN, TEXT(""), (int32)GFrameCounter, 0, 0.0f);
+}
+
+static void EndFrameEvent()
+{
+	AndroidThunkCpp_Engine_SendEvent(EVENTTYPE_FRAME_END, TEXT(""), (int32)(GFrameCounter - 1), 0, 0.0f);
+}
+
+static void PreLoadMapEvent(const FString& Param)
+{
+	AndroidThunkCpp_Engine_SendEvent(EVENTTYPE_PRE_LOAD_MAP, Param, 0, 0, 0.0f);
+}
+
+static void PostLoadMapEvent(UWorld* World)
+{
+	if (World != nullptr)
+	{
+		FString LevelName = World->GetMapName();
+		AndroidThunkCpp_Engine_SendEvent(EVENTTYPE_POST_LOAD_MAP, LevelName, 0, 0, 0.0f);
+	}
+	else
+	{
+		AndroidThunkCpp_Engine_SendEvent(EVENTTYPE_POST_LOAD_MAP, TEXT(""), 0, 0, 0.0f);
+	}
+}
+
+JNI_METHOD void Java_com_epicgames_makeaar_GameActivityForMakeAAR_nativeSetCommandline(JNIEnv* jenv, jobject thiz, jstring commandline)
+{
+	auto UECommandLine = FJavaHelper::FStringFromParam(jenv, commandline);
+	GAndroidCommandLine = UECommandLine;
+}
+
+JNI_METHOD void Java_com_epicgames_makeaar_GameActivityForMakeAAR_nativeMain(JNIEnv* jenv, jobject thiz, jstring projectModule)
+{
+	BootTimingPoint("native_main");
+	GIsGameAgnosticExe = true;
+
+	if (projectModule)
+	{
+		FString name = FJavaHelper::FStringFromParam(jenv, projectModule);
+		STANDALONE_DEBUG_LOGf(LogAndroid, TEXT("Java_com_epicgames_makeaar_GameActivityForMakeAAR_nativeMain : OverrideProjectModule to '%s'"), *name);
+		GEngineLoop.OverrideProjectModule(FJavaHelper::FStringFromParam(jenv, projectModule), name);
+	}
+	else
+	{
+		STANDALONE_DEBUG_LOGf(LogAndroid, TEXT("Java_com_epicgames_makeaar_GameActivityForMakeAAR_nativeMain : use current ProjectModule"));
+
+	}
+
+	// register some delegates we want to pass back
+	FCoreDelegates::OnInit.AddStatic(InitEvent);
+	FCoreDelegates::OnPostEngineInit.AddStatic(PostEngineInitEvent);
+	FCoreDelegates::OnFEngineLoopInitComplete.AddStatic(EngineLoopInitCompleteEvent);
+	FCoreDelegates::OnBeginFrame.AddStatic(BeginFrameEvent);
+	FCoreDelegates::OnEndFrame.AddStatic(EndFrameEvent);
+	FCoreUObjectDelegates::PreLoadMap.AddStatic(PreLoadMapEvent);
+	FCoreUObjectDelegates::PostLoadMapWithWorld.AddStatic(PostLoadMapEvent);
+
+	STANDALONE_DEBUG_LOGf(LogAndroid, TEXT("Entering engine main function"));
+
+	dummy_state.activity = nullptr;
+
+	pthread_attr_t otherAttr;
+	pthread_attr_init(&otherAttr);
+	pthread_attr_setdetachstate(&otherAttr, PTHREAD_CREATE_DETACHED);
+	pthread_create(&G_AndroidMainThread, &otherAttr, android_main, &dummy_state);
+
+
+	STANDALONE_DEBUG_LOGf(LogAndroid, TEXT("Created main thread"));
+}
+
+JNI_METHOD jint Java_com_epicgames_makeaar_GameActivityForMakeAAR_nativeInputTouch(JNIEnv* jenv, jobject thiz, jint device, jint inType, jint inPointerId, jint inX, jint inY)
+{
+	TArray<TouchInput> TouchesArray;
+
+	TouchType type = TouchEnded;
+	switch (inType)
+	{
+	case 0: type = TouchBegan; break;
+	case 1: type = TouchMoved; break;
+	case 2: type = TouchEnded; break;
+	}
+
+	ANativeWindow* Window = (ANativeWindow*)FAndroidWindow::GetHardwareWindow_EventThread();
+	if (!Window)
+	{
+		return 0;
+	}
+
+	int32_t Width = 0;
+	int32_t Height = 0;
+
+	if (Window)
+	{
+		if (GAndroidWindowOverride != nullptr)
+		{
+			Width = ANativeWindow_getWidth((ANativeWindow*)GAndroidWindowOverride);
+			Height = ANativeWindow_getHeight((ANativeWindow*)GAndroidWindowOverride);
+		}
+		else
+		{
+			// we are on the event thread. true here indicates we will retrieve dimensions from the current window.
+			FAndroidWindow::CalculateSurfaceSize(Width, Height, true);
+		}
+	}
+
+	// make sure OpenGL context created before accepting touch events.. FAndroidWindow::GetScreenRect() may try to create it early from wrong thread if this is the first call
+	if (!GAndroidGPUInfoReady)
+	{
+		return 1;
+	}
+	FPlatformRect ScreenRect = FAndroidWindow::GetScreenRect(true);
+	int32_t ScreenWidth = ScreenRect.Right - ScreenRect.Left;
+	int32_t ScreenHeight = ScreenRect.Bottom - ScreenRect.Top;
+
+	int32 eventX = inX - ScreenRect.Left;
+	int32 eventY = inY - ScreenRect.Top;
+	if (eventX < 0 || eventY < 0 || eventX >= ScreenWidth || eventY >= ScreenHeight)
+	{
+		return 0;
+	}
+
+	int pointerId = inPointerId;
+	float x = FMath::Max<float>(FMath::Min<float>((float)eventX / Width, 1.f), 0.f) * (ScreenWidth - 1);
+	float y = FMath::Max<float>(FMath::Min<float>((float)eventY / Height, 1.f), 0.f) * (ScreenHeight - 1);
+
+	//STANDALONE_DEBUG_LOGf(LogAndroid, TEXT("TOUCH: at (%d,%d) windowPos = (%d,%d), windowSize = (%d,%d), screenSize = (%d,%d), final pos: (%f,%f)"),
+	//	eventX, eventY, ScreenRect.Left, ScreenRect.Top, ScreenWidth, ScreenHeight, Width, Height, x, y);
+	UE_LOG(LogAndroid, Verbose, TEXT("Received targeted motion event from pointer %u (id %d) action %d: (%.2f, %.2f)"), 0, pointerId, type, x, y);
+
+	TouchInput TouchMessage;
+	TouchMessage.DeviceId = device;
+	TouchMessage.Handle = pointerId;
+	TouchMessage.Type = type;
+	TouchMessage.Position = FVector2D(x, y);
+	TouchMessage.LastPosition = FVector2D(x, y);
+	TouchesArray.Add(TouchMessage);
+
+	FAndroidInputInterface::QueueTouchInput(TouchesArray);
+
+	return 1;
+}
+
+static bool MountPak(const FString& PakFilePath, int32 PakOrder, const FString& PakMountPoint)
+{
+	// Mount all pak files found in this content
+	FPakPlatformFile* PakFileMgr = (FPakPlatformFile*)(FPlatformFileManager::Get().FindPlatformFile(TEXT("PakFile")));
+	if (PakFileMgr == nullptr)
+	{
+		return false;
+	}
+
+	if (!PakMountPoint.IsEmpty())
+	{
+		const TCHAR* MountPoint = PakMountPoint.GetCharArray().GetData();
+		return PakFileMgr->Mount(*PakFilePath, PakOrder, MountPoint);
+	}
+
+	return PakFileMgr->Mount(*PakFilePath, PakOrder);
+}
+
+JNI_METHOD bool Java_com_epicgames_makeaar_Engine_nativeMountPak(JNIEnv* jenv, jobject thiz, jstring pakFile, jint order, jstring mountPoint)
+{
+	FString PakFilePath = FJavaHelper::FStringFromParam(jenv, pakFile);
+	uint32 PakOrder = (uint32)FMath::Max(0, order);
+	FString PakMountPoint = FJavaHelper::FStringFromParam(jenv, mountPoint);
+
+	if (!IsInGameThread() && FTaskGraphInterface::IsRunning())
+	{
+		STANDALONE_DEBUG_LOGf(LogAndroid, TEXT("nativeMountPak: deferred '%s'"), *PakFilePath);
+		FSimpleDelegateGraphTask::CreateAndDispatchWhenReady(
+			FSimpleDelegateGraphTask::FDelegate::CreateLambda([=]()
+				{
+					MountPak(PakFilePath, PakOrder, PakMountPoint);
+				}), TStatId(), nullptr, ENamedThreads::GameThread);
+		return true;
+	}
+	else
+	{
+		STANDALONE_DEBUG_LOGf(LogAndroid, TEXT("nativeMountPak: immediate '%s'"), *PakFilePath);
+		return MountPak(PakFilePath, PakOrder, PakMountPoint);
+	}
+}
+
+JNI_METHOD void Java_com_epicgames_makeaar_GameActivityForMakeAAR_nativeAppCommand(JNIEnv* jenv, jobject thiz, jint cmd)
+{
+	dummy_state.window = (ANativeWindow*)GAndroidWindowOverride;
+	dummy_state.pendingWindow = (ANativeWindow*)GAndroidWindowOverride;
+
+	bool bForceCommandToGameThread = CVarUseGameThreadForCommands.GetValueOnAnyThread() != 0;
+
+	if (bForceCommandToGameThread && FTaskGraphInterface::IsRunning())
+	{
+		STANDALONE_DEBUG_LOGf(LogAndroid, TEXT("OnAppCommandCB(Queue on GameThread) cmd: %u, tid = %d, FTaskGraphInterface::IsRunning()=%d"), cmd, gettid(), FTaskGraphInterface::IsRunning());
+
+
+		FGraphEventRef nativeAppCommandCalled = FFunctionGraphTask::CreateAndDispatchWhenReady([&]()
+			{
+				OnAppCommandCB(&dummy_state, cmd);
+			}, TStatId(), NULL, ENamedThreads::GameThread);
+	}
+	else
+	{
+		STANDALONE_DEBUG_LOGf(LogAndroid, TEXT("OnAppCommandCB(immediate) cmd: %d, tid = %d, FTaskGraphInterface::IsRunning()=%d"), (int)cmd, gettid(), FTaskGraphInterface::IsRunning());
+
+	    if (GEngine != NULL)
+	    {
+		    // Run on game thread to avoid race condition
+		    AsyncTask(ENamedThreads::GameThread, [cmd]()
+			    {
+				    STANDALONE_DEBUG_LOGf(LogAndroid, TEXT("Issue nativeAppCommand: %d, tid = %d"), (int)cmd, gettid());
+				    OnAppCommandCB(&dummy_state, cmd);
+			    });
+	    }
+	}
+}
+
+#endif // USE_ANDROID_STANDALONE
+
+
+//Wraps __cxa_atexit and avoids registering any destructors for static objects to save startup time
+extern "C" {
+    int __wrap___cxa_atexit(void (*f)(void*), void* objptr, void* dso)
+    {
+	    return 0;
+    }
+};
+
 #endif
diff --git a/Engine/Source/Runtime/Launch/Private/LaunchEngineLoop.cpp b/Engine/Source/Runtime/Launch/Private/LaunchEngineLoop.cpp
--- a/Engine/Source/Runtime/Launch/Private/LaunchEngineLoop.cpp
+++ b/Engine/Source/Runtime/Launch/Private/LaunchEngineLoop.cpp
@@ -1564,6 +1564,17 @@ FEngineLoop::FEngineLoop()
 { }
 
 
+static FString OriginalProjectModuleName;
+static FString ReplacementProjectModuleName;
+
+void FEngineLoop::OverrideProjectModule(const FString& InOriginalProjectModuleName, const FString& InReplacementProjectModuleName)
+{
+	OriginalProjectModuleName = InOriginalProjectModuleName;
+	ReplacementProjectModuleName = InReplacementProjectModuleName;
+	FPlatformMisc::LowLevelOutputDebugStringf(TEXT("OverrideProjectModule : OriginalProjectModuleName=%s, ReplacementProjectModuleName=%s\n"), *OriginalProjectModuleName, *ReplacementProjectModuleName);
+
+}
+
 int32 FEngineLoop::PreInit(int32 ArgC, TCHAR* ArgV[], const TCHAR* AdditionalCommandline)
 {
 	FString CmdLine = FCommandLine::BuildFromArgV(nullptr, ArgC, ArgV, AdditionalCommandline);
@@ -2606,6 +2617,11 @@ int32 FEngineLoop::PreInitPreStartupScreen(const TCHAR* CmdLine)
 			// Add the enterprise binaries directory if we're an enterprise project
 			FModuleManager::Get().AddBinariesDirectory(*FPaths::Combine(FPaths::EnterpriseDir(), TEXT("Binaries"), FPlatformProcess::GetBinariesSubdirectory()), false);
 		}
+
+		if (!ReplacementProjectModuleName.IsEmpty())
+		{
+			IProjectManager::Get().SubstituteModule(OriginalProjectModuleName, ReplacementProjectModuleName);
+		}
 	}
 
 #if !IS_PROGRAM
diff --git a/Engine/Source/Runtime/Launch/Public/LaunchEngineLoop.h b/Engine/Source/Runtime/Launch/Public/LaunchEngineLoop.h
--- a/Engine/Source/Runtime/Launch/Public/LaunchEngineLoop.h
+++ b/Engine/Source/Runtime/Launch/Public/LaunchEngineLoop.h
@@ -87,6 +87,9 @@ public:
 	/** Load core modules. */
 	bool LoadCoreModules();
 
+	/** Allow override of the project module loaded from project file */
+	void OverrideProjectModule(const FString& InOriginalProjectModuleName, const FString& InReplacementProjectModuleName);
+
 	/** Clean up PreInit context. */
 	void CleanupPreInitContext();
 
diff --git a/Engine/Source/Runtime/OpenGLDrv/Private/Android/AndroidEGL.cpp b/Engine/Source/Runtime/OpenGLDrv/Private/Android/AndroidEGL.cpp
--- a/Engine/Source/Runtime/OpenGLDrv/Private/Android/AndroidEGL.cpp
+++ b/Engine/Source/Runtime/OpenGLDrv/Private/Android/AndroidEGL.cpp
@@ -292,10 +292,17 @@ void AndroidEGL::CreateEGLRenderSurface(ANativeWindow* InWindow, bool bCreateWnd
 
 		if (FAndroidPlatformRHIFramePacer::CVarAllowFrameTimestamps.GetValueOnAnyThread())
 		{
+			STANDALONE_DEBUG_LOGf(LogAndroid, TEXT("AndroidEGL::CreateEGLRenderSurface(InWindow = %p) using a.allowFrameTimestamps enable EGL_TIMESTAMPS_ANDROID on %p"), InWindow, PImplData->eglSurface);
 			eglSurfaceAttrib(PImplData->eglDisplay, PImplData->eglSurface, EGL_TIMESTAMPS_ANDROID, EGL_TRUE);
 		}
+		else
+		{
+			// HAD to add the false condition so that android attributes reflect current state of CVar.
+			STANDALONE_DEBUG_LOGf(LogAndroid, TEXT("AndroidEGL::CreateEGLRenderSurface(InWindow = %p) using a.allowFrameTimestamps disable EGL_TIMESTAMPS_ANDROID on %p"), InWindow, PImplData->eglSurface);
+			eglSurfaceAttrib(PImplData->eglDisplay, PImplData->eglSurface, EGL_TIMESTAMPS_ANDROID, EGL_FALSE);
+		}
 
-		FPlatformMisc::LowLevelOutputDebugStringf( TEXT("AndroidEGL::CreateEGLRenderSurface() %p" ), PImplData->eglSurface);
+		STANDALONE_DEBUG_LOGf(LogAndroid, TEXT("AndroidEGL::CreateEGLRenderSurface() %p" ), PImplData->eglSurface);
 
 		if(PImplData->eglSurface == EGL_NO_SURFACE )
 		{
@@ -672,12 +679,21 @@ void AndroidEGL::InitRenderSurface(bool bUseSmallSurface, bool bCreateWndSurface
 		if (PImplData->CachedWindowRect.Right > 0 && PImplData->CachedWindowRect.Bottom > 0)
 		{
 			// If we resumed from a lost window reuse the window size, the game thread will update the window dimensions.
-			FPlatformMisc::LowLevelOutputDebugStringf(TEXT("AndroidEGL::InitRenderSurface, Using CachedWindowRect, width: %d, height %d "), PImplData->CachedWindowRect.Right, PImplData->CachedWindowRect.Bottom);
+			FPlatformMisc::LowLevelOutputDebugStringf(TEXT("AndroidEGL::InitRenderSurface, Using CachedWindowRect, left: %d, top: %d, right: %d, bottom: %d "), PImplData->CachedWindowRect.Left, PImplData->CachedWindowRect.Top, PImplData->CachedWindowRect.Right, PImplData->CachedWindowRect.Bottom);
 			WindowSize = PImplData->CachedWindowRect;
 		}
+#if USE_ANDROID_STANDALONE
+		if (WindowSize.Left != 0 || WindowSize.Top != 0)
+		{
+			STANDALONE_DEBUG_LOGf(LogAndroid, TEXT("AndroidEGL::InitRenderSurface, WARNING!!! WindowSize is offset, left: %d, top: %d, right: %d, bottom: %d "), WindowSize.Left, WindowSize.Top, WindowSize.Right, WindowSize.Bottom);
+		}
+		Width = WindowSize.Right - WindowSize.Left;
+		Height = WindowSize.Bottom - WindowSize.Top;
+#else
 
 		Width = WindowSize.Right;
 		Height = WindowSize.Bottom;
+#endif
 
 		FPlatformMisc::LowLevelOutputDebugStringf(TEXT("AndroidEGL::InitRenderSurface, Using width: %d, height %d "), Width, Height);
 		AndroidThunkCpp_SetDesiredViewSize(Width, Height);
@@ -707,12 +723,21 @@ void AndroidEGL::InitSharedSurface(bool bUseSmallSurface)
 		if (PImplData->CachedWindowRect.Right > 0 && PImplData->CachedWindowRect.Bottom > 0)
 		{
 			// If we resumed from a lost window reuse the window size, the game thread will update the window dimensions.
-			FPlatformMisc::LowLevelOutputDebugStringf(TEXT("AndroidEGL::InitSharedSurface, Using CachedWindowRect, width: %d, height %d "), PImplData->CachedWindowRect.Right, PImplData->CachedWindowRect.Bottom);
+			FPlatformMisc::LowLevelOutputDebugStringf(TEXT("AndroidEGL::InitSharedSurface, Using CachedWindowRect, left: %d, top: %d, right: %d, bottom: %d "), PImplData->CachedWindowRect.Left, PImplData->CachedWindowRect.Top, PImplData->CachedWindowRect.Right, PImplData->CachedWindowRect.Bottom);
 			WindowSize = PImplData->CachedWindowRect;
 		}
+#if USE_ANDROID_STANDALONE
+		if (WindowSize.Left != 0 || WindowSize.Top != 0)
+		{
+			STANDALONE_DEBUG_LOGf(LogAndroid, TEXT("AndroidEGL::InitSharedSurface, WARNING!!! WindowSize is offset, left: %d, top: %d, right: %d, bottom: %d "), WindowSize.Left, WindowSize.Top, WindowSize.Right, WindowSize.Bottom);
+		}
+		Width = WindowSize.Right - WindowSize.Left;
+		Height = WindowSize.Bottom - WindowSize.Top;
+#else
 
 		Width = WindowSize.Right;
 		Height = WindowSize.Bottom;
+#endif
 
 		FPlatformMisc::LowLevelOutputDebugStringf(TEXT("AndroidEGL::InitSharedSurface, Using width: %d, height %d "), Width, Height);
 		AndroidThunkCpp_SetDesiredViewSize(Width, Height);
@@ -1202,7 +1227,9 @@ bool AndroidEGL::IsOfflineSurfaceRequired()
 }
 
 ///
-extern FCriticalSection GAndroidWindowLock;
+//extern FCriticalSection GAndroidWindowLock;
+extern void GAndroidWindowLock_Lock(FString calledBy);
+extern void GAndroidWindowLock_Unlock(FString calledBy);
 
 void BlockOnLostWindowRenderCommand(TSharedPtr<FEvent, ESPMode::ThreadSafe> RTBlockedTrigger)
 {
@@ -1232,7 +1259,7 @@ void BlockOnLostWindowRenderCommand(TSharedPtr<FEvent, ESPMode::ThreadSafe> RTBl
 
 		RTBlockedTrigger->Trigger();
 
-		GAndroidWindowLock.Lock();
+		GAndroidWindowLock_Lock("BlockOnLostWindowRenderCommand Vulkan");
 		UE_LOG(LogAndroid, Log, TEXT("RendererBlock acquired window lock"));
 		const auto& OnReinitWindowCallback = FAndroidMisc::GetOnReInitWindowCallback();
 		if (OnReinitWindowCallback)
@@ -1240,17 +1267,17 @@ void BlockOnLostWindowRenderCommand(TSharedPtr<FEvent, ESPMode::ThreadSafe> RTBl
 			OnReinitWindowCallback(FAndroidWindow::GetHardwareWindow_EventThread());
 			UE_LOG(LogAndroid, Log, TEXT("RendererBlock updating window"));
 		}
-		GAndroidWindowLock.Unlock();
+		GAndroidWindowLock_Unlock("BlockOnLostWindowRenderCommand Vulkan");
 	}
 	else
 	{
 		RunOnGLRenderContextThread([&] {
 			RTBlockedTrigger->Trigger();
-			GAndroidWindowLock.Lock();
+			GAndroidWindowLock_Lock("BlockOnLostWindowRenderCommand");
 			UE_LOG(LogAndroid, Log, TEXT("RendererBlock acquired window lock"));
 			AndroidEGL::GetInstance()->SetRenderContextWindowSurface();
 			UE_LOG(LogAndroid, Log, TEXT("RendererBlock updating window"));
-			GAndroidWindowLock.Unlock();
+			GAndroidWindowLock_Unlock("BlockOnLostWindowRenderCommand");
 		}, true);
 	}
 	UE_LOG(LogAndroid, Log, TEXT("RendererBlock released window lock"));
diff --git a/Engine/Source/Runtime/Projects/Private/ProjectManager.cpp b/Engine/Source/Runtime/Projects/Private/ProjectManager.cpp
--- a/Engine/Source/Runtime/Projects/Private/ProjectManager.cpp
+++ b/Engine/Source/Runtime/Projects/Private/ProjectManager.cpp
@@ -106,6 +106,25 @@ bool FProjectManager::LoadModulesForProject( const ELoadingPhase::Type LoadingPh
 	return bSuccess;
 } 
 
+bool FProjectManager::SubstituteModule(const FString& OriginalModuleName, const FString& NewModuleName)
+{
+	if (!CurrentProject.IsValid())
+	{
+		return false;
+	}
+	const int ModuleCount = CurrentProject->Modules.Num();
+	for (int32 Idx = 0; Idx < ModuleCount; Idx++)
+	{
+		if (CurrentProject->Modules[Idx].Name.IsEqual(*OriginalModuleName))
+		{
+			CurrentProject->Modules[Idx].Name = FName(*NewModuleName);
+			return true;
+		}
+	}
+	return false;
+}
+
+
 #if !IS_MONOLITHIC
 bool FProjectManager::CheckModuleCompatibility(TArray<FString>& OutIncompatibleModules)
 {
diff --git a/Engine/Source/Runtime/Projects/Private/ProjectManager.h b/Engine/Source/Runtime/Projects/Private/ProjectManager.h
--- a/Engine/Source/Runtime/Projects/Private/ProjectManager.h
+++ b/Engine/Source/Runtime/Projects/Private/ProjectManager.h
@@ -20,6 +20,7 @@ public:
 	virtual bool LoadProjectFile( const FString& ProjectFile ) override;
 	virtual bool LoadModulesForProject( const ELoadingPhase::Type LoadingPhase ) override;
 	virtual FLoadingModulesForPhaseEvent& OnLoadingPhaseComplete() override { return OnLoadingPhaseCompleteEvent; }
+	virtual bool SubstituteModule(const FString& OriginalModuleName, const FString& NewModuleName) override;
 #if !IS_MONOLITHIC
 	virtual bool CheckModuleCompatibility( TArray<FString>& OutIncompatibleModules ) override;
 #endif
diff --git a/Engine/Source/Runtime/Projects/Public/Interfaces/IProjectManager.h b/Engine/Source/Runtime/Projects/Public/Interfaces/IProjectManager.h
--- a/Engine/Source/Runtime/Projects/Public/Interfaces/IProjectManager.h
+++ b/Engine/Source/Runtime/Projects/Public/Interfaces/IProjectManager.h
@@ -106,6 +106,15 @@ public:
 	 */
 	virtual bool LoadModulesForProject( const ELoadingPhase::Type LoadingPhase ) = 0;
 
+	/**
+	* Replaces a referenced module with a new name
+	*
+	* @param	OriginalModuleName	Name of module to replace
+	* @param	NewModuleName		Name of replacement module
+	*
+	*/
+	virtual bool SubstituteModule(const FString& OriginalModuleName, const FString& NewModuleName) = 0;
+
 	/**
 	 * Callback for when modules for when LoadModulesForProject() completes loading for a specific phase.
 	 */
