AndroidSingleInstance Window Fix

diff --git a/Engine/Source/Runtime/Launch/Private/Android/LaunchAndroid.cpp b/Engine/Source/Runtime/Launch/Private/Android/LaunchAndroid.cpp
--- a/Engine/Source/Runtime/Launch/Private/Android/LaunchAndroid.cpp
+++ b/Engine/Source/Runtime/Launch/Private/Android/LaunchAndroid.cpp
@@ -1720,8 +1720,14 @@ static void OnAppCommandCB(struct android_app* app, int32_t cmd)
 		STANDALONE_DEBUG_LOGf(LogAndroid, TEXT("Case APP_CMD_INIT_WINDOW"));
 		UE_LOG(LogAndroid, Log, TEXT("Case APP_CMD_INIT_WINDOW"));
 #if USE_ANDROID_STANDALONE
-		check(app->pendingWindow == (ANativeWindow*)GAndroidWindowOverride);
-		FAppEventManager::GetInstance()->HandleWindowCreated_EventThread((ANativeWindow*)GAndroidWindowOverride);
+		{
+			ANativeWindow* win = app->pendingWindow ? app->pendingWindow : app->window;
+			if (GAndroidWindowOverride != NULL)
+			{
+				win = (ANativeWindow*)GAndroidWindowOverride;
+			}
+			FAppEventManager::GetInstance()->HandleWindowCreated_EventThread((ANativeWindow*)win);
+		}
 #else
 		FAppEventManager::GetInstance()->HandleWindowCreated_EventThread(app->pendingWindow);
 #endif
