fix linux compile issue with include upper case lower case name.
#android

diff --git a/Engine/Source/Runtime/Launch/Private/Android/LaunchAndroid.cpp b/Engine/Source/Runtime/Launch/Private/Android/LaunchAndroid.cpp
--- a/Engine/Source/Runtime/Launch/Private/Android/LaunchAndroid.cpp
+++ b/Engine/Source/Runtime/Launch/Private/Android/LaunchAndroid.cpp
@@ -47,7 +47,7 @@
 
 #if USE_ANDROID_STANDALONE
 #include "HAL/FileManager.h"
-#include "HAL/PlatformFilemanager.h"
+#include "HAL/PlatformFileManager.h"
 #include "IPlatformFilePak.h"
 #endif
 
