Fix always attempt to restart the app after receiving onStop command

diff --git a/Engine/Source/Runtime/Launch/Private/Android/LaunchAndroid.cpp b/Engine/Source/Runtime/Launch/Private/Android/LaunchAndroid.cpp
--- a/Engine/Source/Runtime/Launch/Private/Android/LaunchAndroid.cpp
+++ b/Engine/Source/Runtime/Launch/Private/Android/LaunchAndroid.cpp
@@ -1957,7 +1957,6 @@ static void OnAppCommandCB(struct android_app* app, int32_t cmd)
 		DEVELOPER_LOG_COMMANDCB_CASE(APP_CMD_STOP);
 		UE_LOG(LogAndroid, Log, TEXT("Case APP_CMD_STOP"));
 		FAppEventManager::GetInstance()->EnqueueAppEvent(APP_EVENT_STATE_ON_STOP);
-		bShouldRestartFromInterrupt = true;
 		break;
 	case APP_CMD_DESTROY:
 		/**
