fix issue with non ASIS android packages failing to resume from a GAndroidWindowLock_Lock

diff --git a/Engine/Source/Runtime/Launch/Private/Android/LaunchAndroid.cpp b/Engine/Source/Runtime/Launch/Private/Android/LaunchAndroid.cpp
--- a/Engine/Source/Runtime/Launch/Private/Android/LaunchAndroid.cpp
+++ b/Engine/Source/Runtime/Launch/Private/Android/LaunchAndroid.cpp
@@ -1057,7 +1057,9 @@ static void* AndroidEventThreadWorker( void* param )
 
 	// window is initially invalid/locked.
 	UE_LOG(LogAndroid, Log, TEXT("AndroidEventThreadWorker, Initial HW window lock."));
-	//GAndroidWindowLock_Lock("AndroidEventThreadWorker");
+#if !USE_ANDROID_STANDALONE
+	GAndroidWindowLock_Lock("AndroidEventThreadWorker");
+#endif
 
 	DEVELOPER_LOG_COMMANDCB_CASE(AndroidEventThreadWorker_BeforeWhile);
 
@@ -1073,7 +1075,9 @@ static void* AndroidEventThreadWorker( void* param )
 	}
 	DEVELOPER_LOG_COMMANDCB_CASE(AndroidEventThreadWorker_AfterWhile);
 
-	//GAndroidWindowLock_Unlock("AndroidEventThreadWorker");
+#if !USE_ANDROID_STANDALONE
+	GAndroidWindowLock_Unlock("AndroidEventThreadWorker");
+#endif
 
 	UE_LOG(LogAndroid, Log, TEXT("AndroidEventThreadWorker->Exiting"));
 	STANDALONE_DEBUG_LOGf(LogAndroid, TEXT("Exiting AndroidEventThreadWorker"));
