Do cleanup in LaunchAndroid.cpp and remove any logic that was experimental
Cleanup runOnUiThread in SimpleContextWrapper.java
fixes bug for

#android

#changelist validated
#virtualized

diff --git a/Engine/Build/Android/Java/src/com/epicgames/unreal/SimpleContextWrapper.java b/Engine/Build/Android/Java/src/com/epicgames/unreal/SimpleContextWrapper.java
--- a/Engine/Build/Android/Java/src/com/epicgames/unreal/SimpleContextWrapper.java
+++ b/Engine/Build/Android/Java/src/com/epicgames/unreal/SimpleContextWrapper.java
@@ -13,265 +13,227 @@ import android.view.WindowManager;
 import java.lang.ref.PhantomReference;
 import java.lang.ref.WeakReference;
 import java.util.Arrays;
-import java.util.concurrent.BlockingQueue;
-import java.util.concurrent.LinkedBlockingQueue;
-import java.util.concurrent.ThreadPoolExecutor;
-import java.util.concurrent.TimeUnit;
 
 // valid calls for Context...
 // getResources()
 // getApplicationContext()
 // getPackageManager()
 
-public final class SimpleContextWrapper extends ContextWrapper {
+public final class SimpleContextWrapper extends ContextWrapper
+{
 	final Handler mHandler;
 	final private Thread mUiThread;
-	final public Context mAppActivityContext;
+final public Context mAppActivityContext;
 
-	public static Logger Log = new Logger("UE", "SimpleContextWrapper");
+public static Logger Log = new Logger("UE", "SimpleContextWrapper");
 
-	public static Activity getActivityOfContext(Context context) {
-		if (context == null) {
-			return null;
-		}
-
-		while (context instanceof ContextWrapper) {
-			if (context instanceof Activity) {
-				return (Activity) context;
-			}
-			context = ((ContextWrapper) context).getBaseContext();
-		}
+public static Activity getActivityOfContext(Context context)
+{
+	if (context == null)
+	{
 		return null;
+    }
+
+	while (context instanceof ContextWrapper) {
+	if (context instanceof Activity) {
+        return (Activity)context;
+    }
+	context = ((ContextWrapper)context).getBaseContext();
+}
+return null;
 	}
 
-	public static Context getBaseOfContext(Context context) {
-		if (context == null) {
-			return null;
-		}
+	public static Context getBaseOfContext(Context context)
+{
+	if (context == null)
+	{
+		return null;
+    }
+    while (context instanceof ContextWrapper) {
 
-		while (context instanceof ContextWrapper) {
+	if (!(((ContextWrapper)context).getBaseContext() instanceof ContextWrapper)) {
+        return ((ContextWrapper)context).getBaseContext();
+    }
+	context = ((ContextWrapper)context).getBaseContext();
+}
 
-			if ( !(((ContextWrapper) context).getBaseContext() instanceof ContextWrapper)) {
-				return ((ContextWrapper) context).getBaseContext();
-			}
-			context = ((ContextWrapper) context).getBaseContext();
-		}
+return null;
+	}
 
+	public static ContextWrapper getContextWrapperOfContext(Context context)
+{
+	if (context == null)
+	{
 		return null;
 	}
 
-//	public static boolean isSimpleActivityOfContext(Context context) {
-//		if (context == null) {
-//			return false;
-//		}
-//
-//		while (context instanceof ContextWrapper) {
-//			if (context instanceof SimpleContextWrapper) {
-//				return true;
-//			}
-//			context = ((ContextWrapper) context).getBaseContext();
-//		}
-//
-//		return false;
-//	}
-
-	public static ContextWrapper getContextWrapperOfContext(Context context) {
-		if (context == null) {
-			return null;
-		}
-
-		Context searchContext = context;
-		while (searchContext instanceof ContextWrapper) {
-			if (searchContext instanceof SimpleContextWrapper) {
-				return (SimpleContextWrapper)searchContext;
-			}
-			searchContext = ((ContextWrapper) searchContext).getBaseContext();
-		}
-
-
-		searchContext = context;
-		while (searchContext instanceof ContextWrapper) {
-			if (searchContext instanceof Activity) {
-				return (ContextWrapper)searchContext;
-			}
-			searchContext = ((ContextWrapper) searchContext).getBaseContext();
-		}
-
-		while (context instanceof ContextWrapper) {
-			if (context instanceof ContextWrapper) {
-				return (ContextWrapper)context;
-			}
-			context = ((ContextWrapper) context).getBaseContext();
-		}
-
-		return null;
+	Context searchContext = context;
+	while (searchContext instanceof ContextWrapper) {
+	if (searchContext instanceof SimpleContextWrapper) {
+		return (SimpleContextWrapper)searchContext;
+    }
+	searchContext = ((ContextWrapper)searchContext).getBaseContext();
+}
+
+
+searchContext = context;
+while (searchContext instanceof ContextWrapper) {
+	if (searchContext instanceof Activity) {
+		return (ContextWrapper)searchContext;
+    }
+	searchContext = ((ContextWrapper)searchContext).getBaseContext();
+}
+
+while (context instanceof ContextWrapper) {
+    if (context instanceof ContextWrapper) {
+        return (ContextWrapper)context;
+    }
+	context = ((ContextWrapper)context).getBaseContext();
+}
+
+return null;
 	}
 
-	public static SimpleContextWrapper getSimpleActivityOfContext(Context context) {
-		if (context == null) {
-			return null;
-		}
+	public static SimpleContextWrapper getSimpleActivityOfContext(Context context)
+{
+    if (context == null)
+    {
+        return null;
+    }
 
-		while (context instanceof ContextWrapper) {
-			if (context instanceof SimpleContextWrapper) {
-				return (SimpleContextWrapper)context;
-			}
-			context = ((ContextWrapper) context).getBaseContext();
-		}
+	while (context instanceof ContextWrapper) {
+    if (context instanceof SimpleContextWrapper) {
+        return (SimpleContextWrapper)context;
+	}
+	context = ((ContextWrapper)context).getBaseContext();
+}
 
-		return null;
+return null;
 	}
 
 
 //	@UiThread
-	public SimpleContextWrapper(SimpleContextWrapper base) {
-		super(base);
-		this.mHandler = base.mHandler;
-		this.mUiThread = base.mUiThread;
-		this.mAppActivityContext = base.mAppActivityContext;
-	}
+	public SimpleContextWrapper(SimpleContextWrapper base)
+{
+	super(base);
+	this.mHandler = base.mHandler;
+	this.mUiThread = base.mUiThread;
+	this.mAppActivityContext = base.mAppActivityContext;
+}
 
 //	@UiThread
-	public SimpleContextWrapper(Activity base) {
-		super(base);
-		SimpleContextWrapper _base = getSimpleActivityOfContext(base);
-		Log.verbose("UE - SimpleContextWrapper(Activity base),  base = " + base + ", _base= " + _base + ", Thread.currentThread() = " + Thread.currentThread());
-
-		if (_base != null)
-		{
-			this.mHandler = _base.mHandler;
-			this.mUiThread = _base.mUiThread;
-			this.mAppActivityContext = _base;
-		}
-		else
-		{
-			this.mHandler = new Handler();
-			this.mUiThread = Thread.currentThread();
-			this.mAppActivityContext = getActivityOfContext(base);
-		}
+public SimpleContextWrapper(Activity base)
+{
+	super(base);
+	SimpleContextWrapper _base = getSimpleActivityOfContext(base);
+	Log.verbose("UE - SimpleContextWrapper(Activity base),  base = " + base + ", _base= " + _base + ", Thread.currentThread() = " + Thread.currentThread());
 
+	if (_base != null)
+	{
+		this.mHandler = _base.mHandler;
+		this.mUiThread = _base.mUiThread;
+		this.mAppActivityContext = _base;
+	}
+	else
+	{
+		this.mHandler = new Handler();
+		this.mUiThread = Thread.currentThread();
+		this.mAppActivityContext = getActivityOfContext(base);
 	}
 
+}
+
 //	@UiThread
-	public SimpleContextWrapper(Context base) {
-		super(base);
-
-		SimpleContextWrapper _base = getSimpleActivityOfContext(base);
-		Log.verbose("UE - SimpleContextWrapper(Context base),  base = " + base + ", _base= " + _base + ", Thread.currentThread() = " + Thread.currentThread());
-
-		if (_base != null)
-		{
-			this.mHandler = _base.mHandler;
-			this.mUiThread = _base.mUiThread;
-			this.mAppActivityContext = _base.mAppActivityContext;
-		}
-		else
-		{
-			this.mHandler = new Handler();
-			this.mUiThread = Thread.currentThread();
-			this.mAppActivityContext = getActivityOfContext(base);
-		}
-	}
+public SimpleContextWrapper(Context base)
+{
+	super(base);
 
-	public final void runOnUiThread(Runnable action) {
-		Log.verbose("UE - SimpleContextWrapper, trying to runOnUiThread mAppActivityContext = " + mAppActivityContext + ", mUiThread= " + mUiThread + ", Thread.currentThread() = " + Thread.currentThread());
-
-		Activity _activity = mAppActivityContext instanceof Activity ? (Activity) mAppActivityContext : null;
-		if (_activity == null)
-		{
-			_activity = getActivityOfContext(this);
-		}
-		if (_activity == null)
-		{
-			_activity = getActivityOfContext(getBaseOfContext(this));
-		}
-		
-		if (_activity != null)
-		{
-			_activity.runOnUiThread(action);
-		}
-		else 
-		{
-
-			if (Thread.currentThread() != mUiThread) {
-				mHandler.post(action);
-			} else {
-				Log.error("UE - SimpleContextWrapper, trying to runOnUiThread is being BYPASSED and action called on Main Thread, mAppActivityContext = " + mAppActivityContext + ", mUiThread= " + mUiThread + ", Thread.currentThread() = " + Thread.currentThread());
-
-				//action.run();
-				runOnMainThread(action);
-				
-			}
-		}
-	}
+	SimpleContextWrapper _base = getSimpleActivityOfContext(base);
+	Log.verbose("UE - SimpleContextWrapper(Context base),  base = " + base + ", _base= " + _base + ", Thread.currentThread() = " + Thread.currentThread());
 
-	public Window getWindow() {
-		return mAppActivityContext != null ? ((Activity)mAppActivityContext).getWindow() : null;
+	if (_base != null)
+	{
+		this.mHandler = _base.mHandler;
+		this.mUiThread = _base.mUiThread;
+		this.mAppActivityContext = _base.mAppActivityContext;
 	}
-
-	public WindowManager getWindowManager() {
-		return mAppActivityContext != null ?  ((Activity)mAppActivityContext).getWindowManager() : (WindowManager) getSystemService(Context.WINDOW_SERVICE);
+	else
+	{
+		this.mHandler = new Handler();
+		this.mUiThread = Thread.currentThread();
+		this.mAppActivityContext = getActivityOfContext(base);
 	}
-
-	public void addContentView(View view, ViewGroup.LayoutParams params) {
-		if (mAppActivityContext != null) {
-			((Activity)mAppActivityContext).addContentView(view, params);
-		}
+}
+
+public final void runOnUiThread(Runnable action)
+{
+
+	Activity _activity = mAppActivityContext instanceof Activity ? (Activity)mAppActivityContext : null;
+if (_activity == null)
+{
+	_activity = getActivityOfContext(this);
+	if (_activity == null)
+    {
+		_activity = getActivityOfContext(getBaseOfContext(this));
 	}
+}
+
+if (_activity != null)
+{
+	_activity.runOnUiThread(action);
+}
+else
+{
+	Log.debug("UE - SimpleContextWrapper, trying to runOnUiThread mAppActivityContext = " + mAppActivityContext + ", mUiThread= " + mUiThread + ", Thread.currentThread() = " + Thread.currentThread());
 
-	public boolean equals(Context obj) {
-		return (this.hashCode() == obj.hashCode());
+	if (Thread.currentThread() != mUiThread)
+	{
+		mHandler.post(action);
 	}
+	else
+	{
+		Log.error("UE - SimpleContextWrapper, trying to runOnUiThread is being BYPASSED and action called on Main Thread, mAppActivityContext = " + mAppActivityContext + ", mUiThread= " + mUiThread + ", Thread.currentThread() = " + Thread.currentThread());
 
+		//action.run();
+		runOnMainThread(action);
 
-	/*
-	 * Gets the number of available cores
-	 * (not always the same as the maximum number of cores)
-	 */
-	private static int NUMBER_OF_CORES = Runtime.getRuntime().availableProcessors();
-
-	// Instantiates the queue of Runnables as a LinkedBlockingQueue
-	private final static BlockingQueue<Runnable> workQueue = new LinkedBlockingQueue<Runnable>();
-
-	// Sets the amount of time an idle thread waits before terminating
-	private static final int KEEP_ALIVE_TIME = 1;
-	// Sets the Time Unit to seconds
-	private static final TimeUnit KEEP_ALIVE_TIME_UNIT = TimeUnit.SECONDS;
-	
-	static ThreadPoolExecutor threadPoolExecutor;
-
-	private static void initThreadPoolExtractor() {
-
-		// Creates a thread pool manager
-		threadPoolExecutor = new ThreadPoolExecutor(
-			NUMBER_OF_CORES,       // Initial pool size
-			NUMBER_OF_CORES,       // Max pool size
-			KEEP_ALIVE_TIME,
-			KEEP_ALIVE_TIME_UNIT,
-			workQueue
-		);
+	}
+}
 	}
 
+	public Window getWindow()
+{
+	return mAppActivityContext != null ? ((Activity)mAppActivityContext).getWindow() : null;
+}
 
-	public final void runOnMainThread (Runnable action)
+public WindowManager getWindowManager()
+{
+	return mAppActivityContext != null ? ((Activity)mAppActivityContext).getWindowManager() : (WindowManager)getSystemService(Context.WINDOW_SERVICE);
+}
+
+public void addContentView(View view, ViewGroup.LayoutParams params)
+{
+	if (mAppActivityContext != null)
 	{
-		//if (threadPoolExecutor == null) {
-		//	initThreadPoolExtractor();
-		//}
-		//threadPoolExecutor.execute(action);
-
-		if (mHandler == null)
-		{
-			getMainExecutor().execute(action);
-		}
-		else
-		{
-			mHandler.post(action);
-		}
+		((Activity)mAppActivityContext).addContentView(view, params);
+	}
+}
 
-		//action.run();
+public boolean equals(Context obj)
+{
+	return (this.hashCode() == obj.hashCode());
+}
 
-		//mHandler.post(action);
+public final void runOnMainThread(Runnable action)
+{
+	if (mHandler == null)
+	{
+		getMainExecutor().execute(action);
+	}
+	else
+	{
+		mHandler.post(action);
 	}
+}
 
-}
\ No newline at end of file
+}
diff --git a/Engine/Source/Runtime/Launch/Private/Android/LaunchAndroid.cpp b/Engine/Source/Runtime/Launch/Private/Android/LaunchAndroid.cpp
--- a/Engine/Source/Runtime/Launch/Private/Android/LaunchAndroid.cpp
+++ b/Engine/Source/Runtime/Launch/Private/Android/LaunchAndroid.cpp
@@ -224,13 +224,12 @@ static TAutoConsoleVariable<int32> CVarEnableSustainedPerformanceMode(
 	TEXT("  1: Enabled"),
 	ECVF_Default);
 
-
 extern void AndroidThunkCpp_SetSustainedPerformanceMode(bool);
 static void SetSustainedPerformanceMode()
 {
 	static bool bSustainedPerformanceMode = false;
 	bool bIncomingSustainedPerformanceMode = CVarEnableSustainedPerformanceMode.GetValueOnAnyThread() != 0;
-	if (bSustainedPerformanceMode != bIncomingSustainedPerformanceMode)
+	if(bSustainedPerformanceMode != bIncomingSustainedPerformanceMode)
 	{
 		bSustainedPerformanceMode = bIncomingSustainedPerformanceMode;
 		UE_LOG(LogAndroid, Log, TEXT("Setting sustained performance mode: %d"), (int32)bSustainedPerformanceMode);
@@ -287,11 +286,20 @@ static bool bIsResumed = false;
 static bool bShouldRestartFromInterrupt = false;
 static bool bIgnorePauseOnDownloaderStart = false;
 bool bReadyToProcessEvents = false;
+
+#if USE_ANDROID_STANDALONE
 static int GAndroidWindowLockRefCount = 0;
 static const bool bForce_GAndroidWindowLock = true;
+#endif
+
+pthread_t G_AndroidEventThread;
+struct android_app* GNativeAndroidApp = NULL;
 
 void GAndroidWindowLock_Lock(FString calledBy)
 {
+#if !USE_ANDROID_STANDALONE
+	GAndroidWindowLock.Lock();
+#else
 	bool doAction = bForce_GAndroidWindowLock;
 	if (GAndroidWindowLockRefCount >= 0)
 	{
@@ -320,10 +328,14 @@ void GAndroidWindowLock_Lock(FString calledBy)
 		STANDALONE_DEBUG_LOGf(LogAndroid, TEXT("WARNING GAndroidWindowLock_Lock(%s). NOOP: GAndroidWindowLockRefCount=%d, bAppIsActive_EventThread=%d bReadyToProcessEvents=%d"), *calledBy, GAndroidWindowLockRefCount, bAppIsActive_EventThread, bReadyToProcessEvents);
 	}
 	//return GAndroidWindowLockRefCount>0;
+#endif
 }
 
 void GAndroidWindowLock_Unlock(FString calledBy)
 {
+#if !USE_ANDROID_STANDALONE
+	GAndroidWindowLock.Unlock();
+#else
 	bool doAction = bForce_GAndroidWindowLock;
 	if (GAndroidWindowLockRefCount > 0)
 	{
@@ -356,6 +368,7 @@ void GAndroidWindowLock_Unlock(FString calledBy)
 		STANDALONE_DEBUG_LOGf(LogAndroid, TEXT("GAndroidWindowLock_Unlock(%s). NOOP: GAndroidWindowLockRefCount=%d, bAppIsActive_EventThread=%d bReadyToProcessEvents=%d"), *calledBy, GAndroidWindowLockRefCount, bAppIsActive_EventThread, bReadyToProcessEvents);
 	}
 	//return GAndroidWindowLockRefCount==0;
+#endif
 }
 static void SuspendApp_EventThread();
 static void ActivateApp_EventThread();
@@ -367,15 +380,13 @@ static void ActivateApp_EventThread();
 
 void FPlatformMisc::UnlockAndroidWindow()
 {
-	DEVELOPER_LOG_COMMANDCB_CASE(UnlockAndroidWindow);
 	UE_LOG(LogAndroid, Log, TEXT("Unlocking android HW window during preinit. bAppIsActive_EventThread=%d"), bAppIsActive_EventThread);
 
 	check(IsInGameThread());
 	check(FTaskGraphInterface::IsRunning());
-
-	GAndroidWindowLock_Unlock("UnlockAndroidWindow");
 	
 	bReadyToProcessEvents = true;
+	GAndroidWindowLock_Unlock("UnlockAndroidWindow");
 }
 
 JNI_METHOD void Java_com_epicgames_unreal_GameActivity_nativeResumeMainInit(JNIEnv* jenv, jobject thiz)
@@ -549,7 +560,7 @@ extern void AndroidThunkCpp_DismissSplashScreen();
 static void OnNativeWindowResized(ANativeActivity* activity, ANativeWindow* window)
 {
 	static int8_t cmd = APP_CMD_WINDOW_RESIZED;
-	struct android_app* app = (struct android_app*)activity->instance;
+	struct android_app* app = (struct android_app *)activity->instance;
 	write(app->msgwrite, &cmd, sizeof(cmd));
 }
 
@@ -566,7 +577,7 @@ static void ApplyAndroidCompatConfigRules()
 			static const TCHAR AndroidCompat[] = TEXT("AndroidCompat.");
 			if (Key.StartsWith(AndroidCompat))
 			{
-				FString CVarName = Key.Mid(UE_ARRAY_COUNT(AndroidCompat) - 1);
+				FString CVarName = Key.Mid(UE_ARRAY_COUNT(AndroidCompat)-1);
 				if (AllowedCompatCVars.Contains(CVarName))
 				{
 					auto* CVar = IConsoleManager::Get().FindConsoleVariable(*CVarName);
@@ -583,20 +594,11 @@ static void ApplyAndroidCompatConfigRules()
 	}
 }
 
-#if USE_ANDROID_STANDALONE
-static void TickDeferredAppCommands(FString calledBy);
-static void IssueAppCommand(long Command);
-#endif
-
-
 //Main function called from the android entry point
 #if USE_ANDROID_STANDALONE
 void* AndroidMain(void* param)
 {
 	struct android_app* state = (struct android_app*)param;
-	FTaskTagScope Scope(ETaskTag::EGameThread);
-
-	GGameThreadId = FPlatformTLS::GetCurrentThreadId();
 
 #else
 int32 AndroidMain(struct android_app* state)
@@ -607,6 +609,11 @@ int32 AndroidMain(struct android_app* state)
 
 	FPlatformMisc::LowLevelOutputDebugString(TEXT("Entered AndroidMain()\n"));
 
+	FTaskTagScope Scope(ETaskTag::EGameThread);
+	GGameThreadId = FPlatformTLS::GetCurrentThreadId();
+	GNativeAndroidApp = state;
+	check(GNativeAndroidApp);
+
 	// Force the first call to GetJavaEnv() to happen on the game thread, allowing subsequent calls to occur on any thread
 	FAndroidApplication::GetJavaEnv();
 #if !USE_ANDROID_STANDALONE
@@ -622,7 +629,7 @@ int32 AndroidMain(struct android_app* state)
 	// adjust the file descriptor limits to allow as many open files as possible
 	rlimit cur_fd_limit;
 	{
-		int result = getrlimit(RLIMIT_NOFILE, &cur_fd_limit);
+		int result = getrlimit(RLIMIT_NOFILE, & cur_fd_limit);
 		//FPlatformMisc::LowLevelOutputDebugStringf(TEXT("(%d) Current fd limits: soft = %lld, hard = %lld"), result, cur_fd_limit.rlim_cur, cur_fd_limit.rlim_max);
 	}
 	{
@@ -633,7 +640,7 @@ int32 AndroidMain(struct android_app* state)
 		//FPlatformMisc::LowLevelOutputDebugStringf(TEXT("(%d) Setting fd limits: soft = %lld, hard = %lld"), result, new_limit.rlim_cur, new_limit.rlim_max);
 	}
 	{
-		int result = getrlimit(RLIMIT_NOFILE, &cur_fd_limit);
+		int result = getrlimit(RLIMIT_NOFILE, & cur_fd_limit);
 		//FPlatformMisc::LowLevelOutputDebugStringf(TEXT("(%d) Current fd limits: soft = %lld, hard = %lld"), result, cur_fd_limit.rlim_cur, cur_fd_limit.rlim_max);
 	}
 
@@ -642,7 +649,7 @@ int32 AndroidMain(struct android_app* state)
 	// However, it has existed in the so since Honeycomb, query for the symbol
 	// to determine whether to try controller support
 	{
-		void* Lib = dlopen("libandroid.so", 0);
+		void* Lib = dlopen("libandroid.so",0);
 		if (Lib != NULL)
 		{
 			GetAxes = (GetAxesType)dlsym(Lib, "AMotionEvent_getAxisValue");
@@ -674,7 +681,7 @@ int32 AndroidMain(struct android_app* state)
 		AlwaysAllowedKeyCodes.Add(AlwaysAllowedKeyCodesList[i]);
 	}
 
-	const int ValidGamepadKeyCodeCount = sizeof(ValidGamepadKeyCodesList) / sizeof(uint32);
+	const int ValidGamepadKeyCodeCount = sizeof(ValidGamepadKeyCodesList)/sizeof(uint32);
 	for (int i = 0; i < ValidGamepadKeyCodeCount; ++i)
 	{
 		ValidGamepadKeyCodes.Add(ValidGamepadKeyCodesList[i]);
@@ -772,7 +779,6 @@ int32 AndroidMain(struct android_app* state)
 	// Have to use a special initialize if using the PersistentStorageManager
 	IPlatformFile::GetPlatformPhysical().Initialize(nullptr, FCommandLine::Get());
 
-
 	{
 		SCOPED_BOOT_TIMING("Wait for GAndroidWindowLock.Lock()");
 		// wait for a valid window
@@ -781,7 +787,6 @@ int32 AndroidMain(struct android_app* state)
 		// We don't wait for all of preinit to complete as PreLoadScreens will need to process events during preinit.
 
 		UE_LOG(LogAndroid, Log, TEXT("PreInit android HW window lock. bAppIsActive_EventThread=%d"), bAppIsActive_EventThread);
-		DEVELOPER_LOG_COMMANDCB_CASE(PreInitWindowLock);
 		GAndroidWindowLock_Lock("AndroidMain lock");
 	}
 
@@ -895,13 +900,13 @@ int32 AndroidMain(struct android_app* state)
 
 	// exit out!
 	GEngineLoop.Exit();
-
-	UE_LOG(LogAndroid, Log, TEXT("Exiting is over"));
+	
 #if !USE_ANDROID_STANDALONE
-
 	FPlatformMisc::RequestExit(true, TEXT("AndroidMain"));
 #endif
 
+	UE_LOG(LogAndroid, Log, TEXT("Exiting is over"));
+
 #if USE_ANDROID_STANDALONE
 	return nullptr;
 #else
@@ -1076,9 +1081,7 @@ static void* AndroidEventThreadWorker( void* param )
 	//continue to process events until the engine is shutting down
 	while (!IsEngineExitRequested())
 	{
-		//FPlatformMisc::LowLevelOutputDebugString(TEXT("AndroidEventThreadWorker"));
-		extern uint64 GFrameCounter;
-		STANDALONE_DEBUG_LOGf(LogAndroid, TEXT("AndroidEventThreadWorker(started), EventThreadID=%d, GFrameCounter=%u"), EventThreadID, GFrameCounter);
+//		FPlatformMisc::LowLevelOutputDebugString(TEXT("AndroidEventThreadWorker"));
 
 		AndroidProcessEvents(state);
 
@@ -1104,7 +1107,7 @@ static void AndroidProcessEvents(struct android_app* state)
 	int events;
 	struct android_poll_source* source;
 
-	while ((ident = ALooper_pollAll(-1, &fdesc, &events, (void**)&source)) >= 0)
+	while((ident = ALooper_pollAll(-1, &fdesc, &events, (void**)&source)) >= 0)
 	{
 		// process this event
 		if (source)
@@ -1114,10 +1117,6 @@ static void AndroidProcessEvents(struct android_app* state)
 	}
 }
 
-pthread_t G_AndroidEventThread;
-
-struct android_app* GNativeAndroidApp = NULL;
-
 void android_main(struct android_app* state)
 {
 	FTaskTagScope Scope(ETaskTag::EGameThread);
@@ -1298,12 +1297,12 @@ static int32_t HandleInputCB(struct android_app* app, AInputEvent* event)
 			(GetAxes != NULL) &&
 			(actionType == AMOTION_EVENT_ACTION_MOVE))
 		{
-			const int axisCount = sizeof(AxisList) / sizeof(int32_t);
+			const int axisCount = sizeof(AxisList)/sizeof(int32_t);
 
 			// poll all the axes and forward to update controller state
 			for (int axis = 0; axis < axisCount; axis++)
 			{
-				float val = GetAxes(event, AxisList[axis], 0);
+				float val = GetAxes( event, AxisList[axis], 0);
 				FAndroidInputInterface::JoystickAxisEvent(device, AxisList[axis], val);
 			}
 
@@ -1633,9 +1632,10 @@ extern void BlockRendering();
 static void SuspendApp_EventThread()
 {
 	DEVELOPER_LOG_COMMANDCB_CASE(SuspendApp_EventThread);
-	if (!bAppIsActive_EventThread || !bHasFocus)
+	if (!bAppIsActive_EventThread)
 	{
-		STANDALONE_DEBUG_LOGf(LogAndroid, TEXT("SuspendApp_EventThread -> UNEXPECTED: event thread, SuspendApp_EventThread called but aborted, GAndroidWindowLockRefCount=%d, bReadyToProcessEvents=%d, bAppIsActive_EventThread=%d, bHasFocus=%d"), GAndroidWindowLockRefCount, bReadyToProcessEvents, bAppIsActive_EventThread, bHasFocus);
+		UE_LOG(LogAndroid, Warning, TEXT("SuspendApp_EventThread -> event thread, suspending app, bAppIsActive_EventThread is false so early return!"));
+		STANDALONE_DEBUG_LOGf(LogAndroid, TEXT("SuspendApp_EventThread -> UNEXPECTED: event thread, SuspendApp_EventThread called but aborted, bReadyToProcessEvents=%d, bAppIsActive_EventThread=%d, bHasFocus=%d"), bReadyToProcessEvents, bAppIsActive_EventThread, bHasFocus);
 		return;
 	}
 	bAppIsActive_EventThread = false;
@@ -1660,6 +1660,7 @@ static void SuspendApp_EventThread()
 	{
 		FPlatformProcess::ReturnSynchEventToPool(EventToDelete);
 	});
+	EMDoneTrigger->Reset();
 
 	// perform the delegates before the window handle is cleared.
 	// This ensures any tasks that require a window handle will have it before we block the RT on the invalid window.
@@ -1684,13 +1685,15 @@ static void SuspendApp_EventThread()
 
 	FThreadHeartBeat::Get().SuspendHeartBeat(true);
 
+#if !USE_ANDROID_STANDALONE
 	// wait for a period of time before blocking rendering
 	UE_LOG(LogAndroid, Log, TEXT("SuspendApp_EventThread -> , waiting for event manager to process. tid: %d"), FPlatformTLS::GetCurrentThreadId());
-	EMDoneTrigger->Reset();
+	
 	bool bSuccess = EMDoneTrigger->Wait(4000);
 	float ElapsedTimeInMs_EMDoneTrigger_Wait = FPlatformTime::ToMilliseconds(FPlatformTime::Cycles() - StartCycles);
 	UE_CLOG(!bSuccess, LogAndroid, Log, TEXT("SuspendApp_EventThread -> backgrounding callback, not responded in timely manner."));
 	STANDALONE_DEBUG_LOGf(LogAndroid, TEXT("SuspendApp_EventThread -> EMDoneTrigger->Wait, waited '%f' ms"), (float)ElapsedTimeInMs_EMDoneTrigger_Wait);
+#endif
 
 	BlockRendering();
 
@@ -1760,7 +1763,6 @@ static void OnAppCommandCB(struct android_app* app, int32_t cmd)
 		{
 			ActivateApp_EventThread();
 		}
-
 		break;
 	case APP_CMD_TERM_WINDOW:
 		/**
@@ -1889,10 +1891,12 @@ static void OnAppCommandCB(struct android_app* app, int32_t cmd)
 		// trigger focus
 		FAppEventManager::GetInstance()->EnqueueAppEvent(APP_EVENT_STATE_WINDOW_GAINED_FOCUS);
 
+#if USE_ANDROID_STANDALONE
 		if (bHasWindow && bHasFocus && bIsResumed)
 		{
 			ActivateApp_EventThread();
 		}
+#endif
 
 		/*
 		* On the initial loading the restart method must be called immediately
@@ -1936,6 +1940,16 @@ static void OnAppCommandCB(struct android_app* app, int32_t cmd)
 		}
 #endif
 
+		// Restart on resuming if did not complete engine initialization
+		if (!bDidCompleteEngineInit && !bIgnorePauseOnDownloaderStart && bAllowReboot)
+		{
+// 			// only do this if early startup enabled
+// 			FString *EarlyRestart = FAndroidMisc::GetConfigRulesVariable(TEXT("earlyrestart"));
+// 			if (EarlyRestart != NULL && EarlyRestart->Equals("true", ESearchCase::IgnoreCase))
+// 			{
+// 				bShouldRestartFromInterrupt = true;
+// 			}
+		}
 		bIgnorePauseOnDownloaderStart = false;
 
 		/*
