PR #8385: Do not shutdown voice module on app stop

#android

diff --git a/Engine/Source/Runtime/Launch/Private/Android/AndroidEventManager.cpp b/Engine/Source/Runtime/Launch/Private/Android/AndroidEventManager.cpp
--- a/Engine/Source/Runtime/Launch/Private/Android/AndroidEventManager.cpp
+++ b/Engine/Source/Runtime/Launch/Private/Android/AndroidEventManager.cpp
@@ -109,7 +109,7 @@ void FAppEventManager::Tick()
 			break;
 		case APP_EVENT_STATE_ON_STOP:
 			bHaveGame = false;
-			ReleaseMicrophone(true);
+			ReleaseMicrophone(false);
 			break;
 		case APP_EVENT_STATE_ON_PAUSE:
 			FAndroidAppEntry::OnPauseEvent();
@@ -178,6 +178,10 @@ void FAppEventManager::ReleaseMicrophone(bool shuttingDown)
 	if (FModuleManager::Get().IsModuleLoaded("Voice"))
 	{
 		UE_LOG(LogAndroidEvents, Log, TEXT("Android release microphone"));
+		if (shuttingDown)
+		{
+			UE_LOG(LogAndroidEvents, Warning, TEXT("Android release microphone - Voice module shutting down - CANNOT RESUME"));
+		}
 		FModuleManager::Get().UnloadModule("Voice", shuttingDown);
 	}
 }
