Fix FAndroidMisc::GetCellularPreference

#android

diff --git a/Engine/Source/Runtime/Launch/Private/Android/AndroidJNI.cpp b/Engine/Source/Runtime/Launch/Private/Android/AndroidJNI.cpp
--- a/Engine/Source/Runtime/Launch/Private/Android/AndroidJNI.cpp
+++ b/Engine/Source/Runtime/Launch/Private/Android/AndroidJNI.cpp
@@ -1440,7 +1440,7 @@ int32 AndroidThunkCpp_GetCellularPreference()
 	int32 value = 0;
 	if (JNIEnv* Env = FAndroidApplication::GetJavaEnv())
 	{
-		FJavaWrapper::CallIntMethod(Env, FJavaWrapper::GameActivityThis, FJavaWrapper::AndroidThunkJava_GetCellularPreference);
+		value = FJavaWrapper::CallIntMethod(Env, FJavaWrapper::GameActivityThis, FJavaWrapper::AndroidThunkJava_GetCellularPreference);
 	}
 	return value;
 }
