- Fixed Xcode proejct generation/compiling/booting for Vision Pro hardware
- Commandline in Xcode is now used, if we want Swift to set the CL, we'll have to pause the engine startup sooner
- Added auto-launch timer to the SwiftUI window for easier development (don't need to wear the VP) to debug startup crashes, etc
- Depth texture was too big, so halved the resolution of the viewports (work in progress)
- Currently asserting in XR code, or crashing in compositor services code

#changelist validated
#virtualized

diff --git a/Engine/Build/BatchFiles/Mac/XcodeBuild.sh b/Engine/Build/BatchFiles/Mac/XcodeBuild.sh
--- a/Engine/Build/BatchFiles/Mac/XcodeBuild.sh
+++ b/Engine/Build/BatchFiles/Mac/XcodeBuild.sh
@@ -83,7 +83,10 @@ case $PLATFORM in
 	"appletvos")
 		PLATFORM="TVOS"
 	;;
-	"xros"|"xrsimulator")
+	"xros")
+		PLATFORM="VisionOS"
+	;;
+	"xrsimulator")
 		PLATFORM="VisionOS"
 		UBT_ARCHFLAG="-architecture=iossimulator" 
 	;;
diff --git a/Engine/Config/BaseDeviceProfiles.ini b/Engine/Config/BaseDeviceProfiles.ini
--- a/Engine/Config/BaseDeviceProfiles.ini
+++ b/Engine/Config/BaseDeviceProfiles.ini
@@ -265,6 +265,7 @@ iPad14,[5,6]=iPadPro6_129
 AppleTV5,=AppleTV
 AppleTV6,=AppleTV4K
 AppleTV11,=AppleTV2_4K
+RealityDevice14,1=VisionPro
 ; Below are generic fallbacks for unrecognized devices. Add any specific devices above this section.
 iPhone=iPhone15
 iPod=iPodTouch7
diff --git a/Engine/Platforms/VisionOS/Source/Runtime/Launch/Source/UESwift.swift b/Engine/Platforms/VisionOS/Source/Runtime/Launch/Source/UESwift.swift
--- a/Engine/Platforms/VisionOS/Source/Runtime/Launch/Source/UESwift.swift
+++ b/Engine/Platforms/VisionOS/Source/Runtime/Launch/Source/UESwift.swift
@@ -6,14 +6,19 @@ struct SwiftUIView: View {
 
 	@Environment(\.openImmersiveSpace) var openImmersiveSpace
 
+	@State private var buttonDisabled = false
+
 	func Open()
 	{
 		Task
 		{
+			buttonDisabled = true
 			await openImmersiveSpace(id: "ImmersiveSpace")
 		}
 	}
 	
+	let timer = Timer.publish(every: 3, on: .main, in: .common).autoconnect()
+	
 	var body: some View {
 		
 	VStack(spacing: 8) {
@@ -23,9 +28,17 @@ struct SwiftUIView: View {
 	  Button(action: {
 		Open()
 //		onClick()
-	  }, label: {
-		Text("Test Button")
-	  })
+	  }
+	  , label: {
+		  Text("Test Button")
+		})
+		.disabled(buttonDisabled)
+
+	}
+	.onReceive(timer) { t2 in
+		print("timer!")
+		timer.upstream.connect().cancel()
+		Open()
 	}
   }
 }
diff --git a/Engine/Source/Programs/UnrealBuildTool/ProjectFiles/Xcode/XcodeProject.cs b/Engine/Source/Programs/UnrealBuildTool/ProjectFiles/Xcode/XcodeProject.cs
--- a/Engine/Source/Programs/UnrealBuildTool/ProjectFiles/Xcode/XcodeProject.cs
+++ b/Engine/Source/Programs/UnrealBuildTool/ProjectFiles/Xcode/XcodeProject.cs
@@ -1842,7 +1842,7 @@ namespace UnrealBuildTool.XcodeProjectXcconfig
 				else if (Platform == UnrealTargetPlatform.VisionOS)
 				{
 					SDKRoot = "xros";
-					SupportedPlatforms = "xrsimulator"; // xros
+					SupportedPlatforms = "xrsimulator xros";
 					DeploymentTargetKey = "XROS_DEPLOYMENT_TARGET";
 					SupportedDevices = UnrealData.VisionOSProjectSettings!.RuntimeDevices;
 					DeploymentTarget = UnrealData.VisionOSProjectSettings.RuntimeVersion;
diff --git a/Engine/Source/Runtime/ApplicationCore/Private/IOS/IOSAppDelegate.cpp b/Engine/Source/Runtime/ApplicationCore/Private/IOS/IOSAppDelegate.cpp
--- a/Engine/Source/Runtime/ApplicationCore/Private/IOS/IOSAppDelegate.cpp
+++ b/Engine/Source/Runtime/ApplicationCore/Private/IOS/IOSAppDelegate.cpp
@@ -260,8 +260,13 @@ static IOSAppDelegate* CachedDelegate = nil;
 -(id)init
 {
 #if UE_USE_SWIFT_UI_MAIN
-	// will get the command later
-	FIOSCommandLineHelper::InitCommandArgs("");
+
+	NSArray* Arguments = [[NSProcessInfo processInfo] arguments];
+	Arguments = [Arguments subarrayWithRange:NSMakeRange(1, [Arguments count] - 1)];
+	FString CmdLine = [Arguments componentsJoinedByString:@" "];
+
+	FIOSCommandLineHelper::InitCommandArgs(*CmdLine);
+
 #endif
 	self = [super init];
 	CachedDelegate = self;
diff --git a/Engine/Source/Runtime/Launch/Private/IOS/LaunchIOS.cpp b/Engine/Source/Runtime/Launch/Private/IOS/LaunchIOS.cpp
--- a/Engine/Source/Runtime/Launch/Private/IOS/LaunchIOS.cpp
+++ b/Engine/Source/Runtime/Launch/Private/IOS/LaunchIOS.cpp
@@ -615,13 +615,25 @@ void FSwiftAppBootstrap::KickoffWithCompositingLayer(CP_OBJECT_cp_layer_renderer
 	{
 		cp_frame_t SwiftLayerFrame = cp_layer_renderer_query_next_frame(Layer);
 		cp_drawable_t SwiftDrawable = cp_frame_query_drawable(SwiftLayerFrame);
+		for (int ViewIndex = 0; ViewIndex < NumViews; ViewIndex++)
 		{
-			cp_view_t View = cp_drawable_get_view(SwiftDrawable, 0);
+			cp_view_t View = cp_drawable_get_view(SwiftDrawable, ViewIndex);
 			cp_view_texture_map_t TextureMap = cp_view_get_view_texture_map(View);
 			MTLViewport Viewport = cp_view_texture_map_get_viewport(TextureMap);
-			
-			NSValue* VPValue = [NSValue valueWithCGRect:CGRectMake(Viewport.originX, Viewport.originY, Viewport.width, Viewport.height)];
+
+			float ScaleDownHack = 0.5;
+			float X = Viewport.originX * ScaleDownHack;
+			float Y = Viewport.originY * ScaleDownHack;
+			float W = Viewport.width * ScaleDownHack;
+			float H = Viewport.height * ScaleDownHack;
+			// both views are at 0? if so, hack the view 1 to be to the right of view 0
+			if (ViewIndex == 1 && X == 0)
+			{
+				X = W;
+			}
+			NSValue* VPValue = [NSValue valueWithCGRect:CGRectMake(X, Y, W, H)];
 			[Viewports addObject:VPValue];
+			FPlatformMisc::LowLevelOutputDebugStringf(TEXT("Adding eye viewport : [%f. %f] / [%f x %f] (scaled by %f\n"), X, Y, W, H, ScaleDownHack);
 		}
 		cp_frame_start_submission(SwiftLayerFrame);
 		id<MTLDevice> Device = cp_layer_renderer_get_device(Layer);
@@ -634,16 +646,6 @@ void FSwiftAppBootstrap::KickoffWithCompositingLayer(CP_OBJECT_cp_layer_renderer
 
 	// cache the viewports in the delegate so code later can get it when asking about the screen bounds
 	AppDelegate.SwiftLayerViewports = Viewports;
-
-	// grab the commandline, skipping over the executable path (index 0), and pass it to unreal
-	NSArray* Arguments = [[NSProcessInfo processInfo] arguments];
-	Arguments = [Arguments subarrayWithRange:NSMakeRange(1, [Arguments count] - 1)];
-	GSavedCommandLine = [Arguments componentsJoinedByString:@" "];
-
-		// Xcode strips quotes, so replace ^ with "
-	GSavedCommandLine = GSavedCommandLine.Replace(TEXT("+"), TEXT("\""));
-	
-	FIOSCommandLineHelper::InitCommandArgs(FString());
 	
 	CGRect FirstViewport = [[AppDelegate.SwiftLayerViewports firstObject] CGRectValue];
 	
