Restore how app_cmd_init and restore is handled (remove ASIS STANDALONE ifdefs that are no longer needed)
#android

diff --git a/Engine/Source/Runtime/ApplicationCore/Private/Android/AndroidWindow.cpp b/Engine/Source/Runtime/ApplicationCore/Private/Android/AndroidWindow.cpp
--- a/Engine/Source/Runtime/ApplicationCore/Private/Android/AndroidWindow.cpp
+++ b/Engine/Source/Runtime/ApplicationCore/Private/Android/AndroidWindow.cpp
@@ -299,7 +299,7 @@ void FAndroidWindow::AcquireWindowRef(ANativeWindow* InWindow)
 	// This logic is to deal with the fact Android lifecycles for activities can overlap. ideally we would create a context based container to manage this
 	// but for now this is a useful protection.
 	GAcquiredWindow = InWindow;
-	STANDALONE_DEBUG_LOG(TEXT("FAndroidWindow::AcquireWindowRef overrode window: %p"), GAndroidWindowOverride);
+	STANDALONE_DEBUG_LOG(TEXT("FAndroidWindow::AcquireWindowRef GAcquiredWindow=%p, GAndroidWindowOverride=%p"), GAcquiredWindow, GAndroidWindowOverride);
 
 #else
 	STANDALONE_DEBUG_LOG(TEXT("AcquireWindowRef USE_ANDROID_JNI is NOT enabled: %p"), InWindow);
diff --git a/Engine/Source/Runtime/Launch/Private/Android/LaunchAndroid.cpp b/Engine/Source/Runtime/Launch/Private/Android/LaunchAndroid.cpp
--- a/Engine/Source/Runtime/Launch/Private/Android/LaunchAndroid.cpp
+++ b/Engine/Source/Runtime/Launch/Private/Android/LaunchAndroid.cpp
@@ -141,7 +141,11 @@ extern "C"
 }
 #endif
 
+#if USE_ANDROID_STANDALONE
+int32 GAndroidEnableNativeResizeEvent = 0;
+#else
 int32 GAndroidEnableNativeResizeEvent = 1;
+#endif
 static FAutoConsoleVariableRef CVarEnableResizeNativeEvent(
 	TEXT("Android.EnableNativeResizeEvent"),
 	GAndroidEnableNativeResizeEvent,
@@ -327,7 +331,6 @@ void GAndroidWindowLock_Lock(FString calledBy)
 	{
 		STANDALONE_DEBUG_LOGf(LogAndroid, TEXT("WARNING GAndroidWindowLock_Lock(%s). NOOP: GAndroidWindowLockRefCount=%d, bAppIsActive_EventThread=%d bReadyToProcessEvents=%d"), *calledBy, GAndroidWindowLockRefCount, bAppIsActive_EventThread, bReadyToProcessEvents);
 	}
-	//return GAndroidWindowLockRefCount>0;
 #endif
 }
 
@@ -360,14 +363,12 @@ void GAndroidWindowLock_Unlock(FString calledBy)
 	if (doAction)
 	{
 		GAndroidWindowLock.Unlock();
-		//bAppIsActive_EventThread = true;
 		STANDALONE_DEBUG_LOGf(LogAndroid, TEXT("GAndroidWindowLock_Unlock(%s). UNLOCK, GAndroidWindowLockRefCount=%d, bAppIsActive_EventThread=%d"), *calledBy, GAndroidWindowLockRefCount, bAppIsActive_EventThread);
 	}
 	else
 	{
 		STANDALONE_DEBUG_LOGf(LogAndroid, TEXT("GAndroidWindowLock_Unlock(%s). NOOP: GAndroidWindowLockRefCount=%d, bAppIsActive_EventThread=%d bReadyToProcessEvents=%d"), *calledBy, GAndroidWindowLockRefCount, bAppIsActive_EventThread, bReadyToProcessEvents);
 	}
-	//return GAndroidWindowLockRefCount==0;
 #endif
 }
 static void SuspendApp_EventThread();
@@ -809,13 +810,11 @@ int32 AndroidMain(struct android_app* state)
 #endif
 	}
 
-#if !USE_ANDROID_STANDALONE
 	// register callback for native window resize
 	if (GAndroidEnableNativeResizeEvent)
 	{
 		state->activity->callbacks->onNativeWindowResized = OnNativeWindowResized;
 	}
-#endif
 
 	// initialize HMDs
 	{
@@ -869,11 +868,7 @@ int32 AndroidMain(struct android_app* state)
 		FAndroidStats::UpdateAndroidStats();
 
 		FAppEventManager::GetInstance()->Tick();
-		if (!FAppEventManager::GetInstance()->IsGamePaused()
-#if USE_ANDROID_STANDALONE
-			&& FAppEventManager::GetInstance()->IsGameInFocus()
-#endif
-			)
+		if (!FAppEventManager::GetInstance()->IsGamePaused())
 		{
 			GEngineLoop.Tick();
 		}
@@ -1071,9 +1066,7 @@ static void* AndroidEventThreadWorker( void* param )
 
 	// window is initially invalid/locked.
 	UE_LOG(LogAndroid, Log, TEXT("AndroidEventThreadWorker, Initial HW window lock."));
-#if !USE_ANDROID_STANDALONE
 	GAndroidWindowLock_Lock("AndroidEventThreadWorker");
-#endif
 
 	DEVELOPER_LOG_COMMANDCB_CASE(AndroidEventThreadWorker_BeforeWhile);
 
@@ -1089,9 +1082,7 @@ static void* AndroidEventThreadWorker( void* param )
 	}
 	DEVELOPER_LOG_COMMANDCB_CASE(AndroidEventThreadWorker_AfterWhile);
 
-#if !USE_ANDROID_STANDALONE
 	GAndroidWindowLock_Unlock("AndroidEventThreadWorker");
-#endif
 
 	UE_LOG(LogAndroid, Log, TEXT("AndroidEventThreadWorker->Exiting"));
 	STANDALONE_DEBUG_LOGf(LogAndroid, TEXT("Exiting AndroidEventThreadWorker"));
@@ -1619,9 +1610,7 @@ static void ActivateApp_EventThread()
 
 	FThreadHeartBeat::Get().ResumeHeartBeat(true);
 
-#if !USE_ANDROID_STANDALONE
 	FPreLoadScreenManager::EnableRendering(true);
-#endif
 
 	extern void AndroidThunkCpp_ShowHiddenAlertDialog();
 	AndroidThunkCpp_ShowHiddenAlertDialog();	
@@ -1679,13 +1668,10 @@ static void SuspendApp_EventThread()
 
 	FEmbeddedCommunication::WakeGameThread();
 
-#if !USE_ANDROID_STANDALONE
 	FPreLoadScreenManager::EnableRendering(false);
-#endif
 
 	FThreadHeartBeat::Get().SuspendHeartBeat(true);
 
-#if !USE_ANDROID_STANDALONE
 	// wait for a period of time before blocking rendering
 	UE_LOG(LogAndroid, Log, TEXT("SuspendApp_EventThread -> , waiting for event manager to process. tid: %d"), FPlatformTLS::GetCurrentThreadId());
 	
@@ -1693,7 +1679,6 @@ static void SuspendApp_EventThread()
 	float ElapsedTimeInMs_EMDoneTrigger_Wait = FPlatformTime::ToMilliseconds(FPlatformTime::Cycles() - StartCycles);
 	UE_CLOG(!bSuccess, LogAndroid, Log, TEXT("SuspendApp_EventThread -> backgrounding callback, not responded in timely manner."));
 	STANDALONE_DEBUG_LOGf(LogAndroid, TEXT("SuspendApp_EventThread -> EMDoneTrigger->Wait, waited '%f' ms"), (float)ElapsedTimeInMs_EMDoneTrigger_Wait);
-#endif
 
 	BlockRendering();
 
@@ -1739,26 +1724,25 @@ static void OnAppCommandCB(struct android_app* app, int32_t cmd)
 		 */
 		// get the window ready for showing
 		DEVELOPER_LOG_COMMANDCB_CASE(APP_CMD_INIT_WINDOW);
-		STANDALONE_DEBUG_LOGf(LogAndroid, TEXT("Case APP_CMD_INIT_WINDOW"));
 		UE_LOG(LogAndroid, Log, TEXT("Case APP_CMD_INIT_WINDOW"));
-		// give scope to local variable in case statement...
-		{
-			ANativeWindow* DimensionWindow = app->pendingWindow ? app->pendingWindow : app->window;
+
 #if USE_ANDROID_STANDALONE
+		{
+			ANativeWindow* win = app->pendingWindow ? app->pendingWindow : app->window;
+			STANDALONE_DEBUG_LOGf(LogAndroid, TEXT("Case APP_CMD_INIT_WINDOW(STANDALONE), app->pendingWindow=%p, app->window=%p, GAndroidWindowOverride=%p"), app->pendingWindow, app->window, GAndroidWindowOverride);
 			if (GAndroidWindowOverride != NULL)
 			{
-				DimensionWindow = (ANativeWindow*)GAndroidWindowOverride;
+				win = (ANativeWindow*)GAndroidWindowOverride;
 			}
-#endif
-
-			FAppEventManager::GetInstance()->HandleWindowCreated_EventThread(DimensionWindow);
-			bHasWindow = DimensionWindow != nullptr;
+			checkf(win != NULL, TEXT("Engine APP_CMD_INIT_WINDOW with STANDALONE Failed since win is NULL!"));
+			FAppEventManager::GetInstance()->HandleWindowCreated_EventThread((ANativeWindow*)win);
 		}
-
-
+#else
+		STANDALONE_DEBUG_LOGf(LogAndroid, TEXT("Case APP_CMD_INIT_WINDOW"));
+		FAppEventManager::GetInstance()->HandleWindowCreated_EventThread(app->pendingWindow);
+#endif
+		bHasWindow = true;
 		UE_LOG(LogAndroid, Log, TEXT("Case APP_CMD_INIT_WINDOW: bHasWindow=%d, bHasFocus=%d, bIsResumed=%d"), bHasWindow, bHasFocus, bIsResumed);
-
-
 		if (bHasWindow && bHasFocus && bIsResumed)
 		{
 			ActivateApp_EventThread();
@@ -1884,6 +1868,11 @@ static void OnAppCommandCB(struct android_app* app, int32_t cmd)
 		bHasFocus = true;
 		DEVELOPER_LOG_COMMANDCB_CASE(APP_CMD_RESUME);
 
+		if (bHasWindow && bHasFocus && bIsResumed)
+		{
+			ActivateApp_EventThread();
+		}
+
 		STANDALONE_DEBUG_LOGf(LogAndroid, TEXT("Case APP_CMD_RESUME, bShouldRestartFromInterrupt=%d"), bShouldRestartFromInterrupt);
 		UE_LOG(LogAndroid, Log, TEXT("Case APP_CMD_RESUME"));
 		FAppEventManager::GetInstance()->EnqueueAppEvent(APP_EVENT_STATE_ON_RESUME);
@@ -1891,13 +1880,6 @@ static void OnAppCommandCB(struct android_app* app, int32_t cmd)
 		// trigger focus
 		FAppEventManager::GetInstance()->EnqueueAppEvent(APP_EVENT_STATE_WINDOW_GAINED_FOCUS);
 
-#if USE_ANDROID_STANDALONE
-		if (bHasWindow && bHasFocus && bIsResumed)
-		{
-			ActivateApp_EventThread();
-		}
-#endif
-
 		/*
 		* On the initial loading the restart method must be called immediately
 		* in order to restart the app if the startup movie was playing
