Modified the launching executable name check to allow for any matching executable editor name instead of just the first one.  This better supports multiple targets.

#rnx

diff --git a/Engine/Source/Runtime/Launch/Private/LaunchEngineLoop.cpp b/Engine/Source/Runtime/Launch/Private/LaunchEngineLoop.cpp
--- a/Engine/Source/Runtime/Launch/Private/LaunchEngineLoop.cpp
+++ b/Engine/Source/Runtime/Launch/Private/LaunchEngineLoop.cpp
@@ -6224,6 +6224,26 @@ void CleanUpPath(FString& InPath)
 	FPaths::RemoveDuplicateSlashes(InPath);
 }
 
+bool DoesCurrentExecutableMatch(const FString& EditorTargetFileName)
+{
+	// Figure out the executable that we should be running
+	FString LaunchExecutableName;
+	FTargetReceipt Receipt;
+	if (!FPaths::FileExists(EditorTargetFileName) || !Receipt.Read(EditorTargetFileName))
+	{
+		return false;
+	}
+	LaunchExecutableName = Receipt.Launch;
+	CleanUpPath(LaunchExecutableName);
+
+	// Get the current executable name. Don't allow relaunching if we're running the console app.
+	FString CurrentExecutableName = FPlatformProcess::ExecutablePath();
+	CleanUpPath(CurrentExecutableName);
+
+	// Nothing to do if they're the same
+	return FPaths::IsSamePath(LaunchExecutableName, CurrentExecutableName);
+}
+
 bool LaunchCorrectEditorExecutable(const FString& EditorTargetFileName)
 {
 	// Don't allow relaunching the executable if we're running some unattended scripted process.
@@ -6413,26 +6433,42 @@ bool FEngineLoop::AppInit( )
 		// Find the editor target
 		FString EditorTargetFileName;
 		FString DefaultEditorTarget;
+		bool bIsRunningExe = false;
 		GConfig->GetString(TEXT("/Script/BuildSettings.BuildSettings"), TEXT("DefaultEditorTarget"), DefaultEditorTarget, GEngineIni);
 
+		// If we have multiple targets for the same code, then try to find a target that matches the running executable.  Otherwise
+		// we use the first matching target.
 		for (const FTargetInfo& Target : FDesktopPlatformModule::Get()->GetTargetsForProject(FPaths::GetProjectFilePath()))
 		{
-			if (Target.Type == EBuildTargetType::Editor && (DefaultEditorTarget.Len() == 0 || Target.Name == DefaultEditorTarget))
+			if (Target.Type == EBuildTargetType::Editor)
 			{
+				FString ScratchEditorTargetFileName;
 				if (FPaths::IsUnderDirectory(Target.Path, FPlatformMisc::ProjectDir()))
 				{
-					EditorTargetFileName = FTargetReceipt::GetDefaultPath(FPlatformMisc::ProjectDir(), *Target.Name, FPlatformProcess::GetBinariesSubdirectory(), FApp::GetBuildConfiguration(), nullptr);
+					ScratchEditorTargetFileName = FTargetReceipt::GetDefaultPath(FPlatformMisc::ProjectDir(), *Target.Name, FPlatformProcess::GetBinariesSubdirectory(), FApp::GetBuildConfiguration(), nullptr);
 				}
 				else if (FPaths::IsUnderDirectory(Target.Path, FPaths::EngineDir()))
 				{
-					EditorTargetFileName = FTargetReceipt::GetDefaultPath(*FPaths::EngineDir(), *Target.Name, FPlatformProcess::GetBinariesSubdirectory(), FApp::GetBuildConfiguration(), nullptr);
+					ScratchEditorTargetFileName = FTargetReceipt::GetDefaultPath(*FPaths::EngineDir(), *Target.Name, FPlatformProcess::GetBinariesSubdirectory(), FApp::GetBuildConfiguration(), nullptr);
+				}
+				if (DoesCurrentExecutableMatch(ScratchEditorTargetFileName))
+				{
+					EditorTargetFileName = MoveTemp(ScratchEditorTargetFileName);
+					bIsRunningExe = true;
+					break;
+				}
+				else if (DefaultEditorTarget.Len() == 0 || Target.Name == DefaultEditorTarget)
+				{
+					if (EditorTargetFileName.Len() == 0)
+					{
+						EditorTargetFileName = MoveTemp(ScratchEditorTargetFileName);
+					}
 				}
-				break;
 			}
 		}
 
 		// If we're not running the correct executable for the current target, and the listed executable exists, run that instead
-		if(LaunchCorrectEditorExecutable(EditorTargetFileName))
+		if (!bIsRunningExe && LaunchCorrectEditorExecutable(EditorTargetFileName))
 		{
 			return false;
 		}
