Movie Render Queue:  fix for motion vectors being lost when modifying the material between renders, leading to loss of motion blur and other artifacts.

Major fix is to ensure use of a consistent GFrameCounter value in animation and rendering with the rest of the frame, by moving increment of this value until after the OnEndFrame delegate, where MRQ rendering occurs.  CSV stat generation end frame call is moved after the GFrameCounter increment, to preserve a side effect when the code calls PlatformMemoryHelpers::GetFrameMemoryStats, which caches platform memory by frame, so that continues to happen at a consistent time (will always be the first call for a given GFrameCounter value, which was the case before).  Removed support for OverrideFrameCounter value in FSceneViewFamily previously used as a workaround (the value is still in the structure, but no longer set or used).

A further minor fix is necessary, which involves moving logic that sets "PreviousBoneTransformRevisionNumber" to "CurrentBoneTransformRevisionNumber" so it happens in UpdateBoneTransformRevisionNumber instead of CreateRenderState_Concurrent and SendRenderDynamicData_Concurrent.  This makes it so additional calls to CreateRenderState_Concurrent in a frame (due to material changes) see the same value for "PreviousBoneTransformRevisionNumber".  Animation updates can be called more than once during Tick, so a frame number is added to track when bone transforms were updated, so the previous revision number is only updated once.  We also send that frame number to the render state instead of the current GFrameCounter.  With these changes, render state exactly mirrors the component state, producing consistent results following re-creation.

#rnx

diff --git a/Engine/Plugins/MovieScene/MovieRenderPipeline/Source/MovieRenderPipelineRenderPasses/Private/Graph/Renderers/MovieGraphImagePassBase.cpp b/Engine/Plugins/MovieScene/MovieRenderPipeline/Source/MovieRenderPipelineRenderPasses/Private/Graph/Renderers/MovieGraphImagePassBase.cpp
--- a/Engine/Plugins/MovieScene/MovieRenderPipeline/Source/MovieRenderPipelineRenderPasses/Private/Graph/Renderers/MovieGraphImagePassBase.cpp
+++ b/Engine/Plugins/MovieScene/MovieRenderPipeline/Source/MovieRenderPipelineRenderPasses/Private/Graph/Renderers/MovieGraphImagePassBase.cpp
@@ -157,12 +157,6 @@ void FMovieGraphImagePassBase::ApplyMovieGraphOverridesToViewFamily(TSharedRef<F
 	// InOutFamily->ViewMode = ViewModeIndex;
 	InOutFamily->bOverrideVirtualTextureThrottle = true;
 	
-	// This is for a workaround involving GFrameCounter already being incremented by the time MRQ renders are submitted
-	// meaning that it's out of sync with some parts of the renderer's idea of the current frame.
-	PRAGMA_DISABLE_DEPRECATION_WARNINGS
-	InOutFamily->OverrideFrameCounter = UE::MovieRenderPipeline::GetRendererFrameCount();
-	PRAGMA_ENABLE_DEPRECATION_WARNINGS
-	
 	// We need to check if this is the first FSceneView being submitted to the renderer module, and set some flags on the ViewFamily for ensuring some
 	// parts of the renderer only get updated once per frame. Kept as an if/else statement to avoid the confusion with setting all of these values to 
 	// some permutation of !/!!bHasRenderedFirstViewThisFrame.
diff --git a/Engine/Plugins/MovieScene/MovieRenderPipeline/Source/MovieRenderPipelineRenderPasses/Private/MoviePipelineImagePassBase.cpp b/Engine/Plugins/MovieScene/MovieRenderPipeline/Source/MovieRenderPipelineRenderPasses/Private/MoviePipelineImagePassBase.cpp
--- a/Engine/Plugins/MovieScene/MovieRenderPipeline/Source/MovieRenderPipelineRenderPasses/Private/MoviePipelineImagePassBase.cpp
+++ b/Engine/Plugins/MovieScene/MovieRenderPipeline/Source/MovieRenderPipelineRenderPasses/Private/MoviePipelineImagePassBase.cpp
@@ -195,10 +195,6 @@ TSharedPtr<FSceneViewFamilyContext> UMoviePipelineImagePassBase::CalculateViewFa
 	OutViewFamily->ViewMode = ViewModeIndex;
 	OutViewFamily->bOverrideVirtualTextureThrottle = true;
 	
-	PRAGMA_DISABLE_DEPRECATION_WARNINGS
-	OutViewFamily->OverrideFrameCounter = UE::MovieRenderPipeline::GetRendererFrameCount();
-	PRAGMA_ENABLE_DEPRECATION_WARNINGS
-	
 	// Kept as an if/else statement to avoid the confusion with setting all of these values to some permutation of !/!!bHasRenderedFirstViewThisFrame.
 	if (!GetPipeline()->bHasRenderedFirstViewThisFrame)
 	{
diff --git a/Engine/Source/Runtime/Engine/Classes/Components/SkinnedMeshComponent.h b/Engine/Source/Runtime/Engine/Classes/Components/SkinnedMeshComponent.h
--- a/Engine/Source/Runtime/Engine/Classes/Components/SkinnedMeshComponent.h
+++ b/Engine/Source/Runtime/Engine/Classes/Components/SkinnedMeshComponent.h
@@ -296,6 +296,12 @@ public:
 		return (LeaderPoseComponentPtr ? LeaderPoseComponentPtr->PreviousBoneTransformRevisionNumber : PreviousBoneTransformRevisionNumber);
 	}
 
+	uint32 GetCurrentBoneTransformFrame() const
+	{
+		const USkinnedMeshComponent* LeaderPoseComponentPtr = LeaderPoseComponent.Get();
+		return (LeaderPoseComponentPtr ? LeaderPoseComponentPtr->CurrentBoneTransformFrame : CurrentBoneTransformFrame);
+	}
+
 	/* this update renderer with new revision number twice so to clear bone velocity for motion blur or temporal AA */
 	ENGINE_API void ClearMotionVector();
 	
@@ -342,6 +348,9 @@ protected:
 	/** current bone transform revision number */
 	uint32 CurrentBoneTransformRevisionNumber;
 
+	/** Stores GFrameCounter when CurrentBoneTransformRevisionNumber was updated, to make sure it's updated once per frame */
+	uint32 CurrentBoneTransformFrame;
+
 	/** Incremented every time the leader bone map changes. Used to keep in sync with any duplicate data needed by other threads */
 	int32 LeaderBoneMapCacheCount;
 
diff --git a/Engine/Source/Runtime/Engine/Private/Components/SkinnedMeshComponent.cpp b/Engine/Source/Runtime/Engine/Private/Components/SkinnedMeshComponent.cpp
--- a/Engine/Source/Runtime/Engine/Private/Components/SkinnedMeshComponent.cpp
+++ b/Engine/Source/Runtime/Engine/Private/Components/SkinnedMeshComponent.cpp
@@ -437,6 +437,7 @@ USkinnedMeshComponent::USkinnedMeshComponent(const FObjectInitializer& ObjectIni
 
 	PreviousBoneTransformRevisionNumber = 0;
 	CurrentBoneTransformRevisionNumber = 0;
+	CurrentBoneTransformFrame = 0;
 
 	ExternalInterpolationAlpha = 0.0f;
 	ExternalDeltaTime = 0.0f;
@@ -872,6 +873,13 @@ void USkinnedMeshComponent::RefreshExternalMorphTargetWeights(bool bZeroOldWeigh
 
 void USkinnedMeshComponent::UpdateBoneTransformRevisionNumber()
 {
+	if (CurrentBoneTransformFrame != GFrameCounter)
+	{
+		// We only want to update the previous revision number for the first call to UpdateBoneTransformRevisionNumber on a given frame
+		PreviousBoneTransformRevisionNumber = CurrentBoneTransformRevisionNumber;
+		CurrentBoneTransformFrame = GFrameCounter;
+	}
+
 	if (BoneTransformUpdateMethodQueue.Last() == EBoneTransformUpdateMethod::ClearMotionVector)
 	{
 		// Last entry is ClearMotionVector, increment revision number by 2 which allows current bone buffer to be bound to previous shader slot to cancel out velocity
@@ -1044,8 +1052,6 @@ void USkinnedMeshComponent::CreateRenderState_Concurrent(FRegisterComponentConte
 			}
 		}
 
-		PreviousBoneTransformRevisionNumber = CurrentBoneTransformRevisionNumber;
-
 		// scene proxy update of material usage based on active morphs
 		UpdateMorphMaterialUsageOnProxy();
 	}
@@ -1124,7 +1130,6 @@ void USkinnedMeshComponent::SendRenderDynamicData_Concurrent()
 	EPreviousBoneTransformUpdateMode PreviousBoneTransformUpdateMode = GetPreviousBoneTransformUpdateMode();
 	// CurrentBoneTransformRevisionNumber and PreviousBoneTransformUpdateMode are up-to-date at this point, it is safe to reset BoneTransformUpdateMethodQueue
 	BoneTransformUpdateMethodQueue.Reset();
-	PreviousBoneTransformRevisionNumber = CurrentBoneTransformRevisionNumber;
 
 	// if we have not updated the transforms then no need to send them to the rendering thread
 	// @todo GIsEditor used to be bUpdateSkelWhenNotRendered. Look into it further to find out why it doesn't update animations in the AnimSetViewer, when a level is loaded in UED (like POC_Cover.gear).
diff --git a/Engine/Source/Runtime/Engine/Private/SkeletalRenderGPUSkin.cpp b/Engine/Source/Runtime/Engine/Private/SkeletalRenderGPUSkin.cpp
--- a/Engine/Source/Runtime/Engine/Private/SkeletalRenderGPUSkin.cpp
+++ b/Engine/Source/Runtime/Engine/Private/SkeletalRenderGPUSkin.cpp
@@ -402,7 +402,9 @@ void FSkeletalMeshObjectGPUSkin::Update(
 		PreviousBoneTransformUpdateMode,
 		InExternalMorphWeightData);
 
-	uint64 FrameNumberToPrepare = GFrameCounter;
+	// The goal is to mirror the exact component state to the render state, including the frame number the component was
+	// updated (as opposed to the current frame number), unless there's no associated component.
+	uint64 FrameNumberToPrepare = InMeshComponent ? InMeshComponent->GetCurrentBoneTransformFrame() : GFrameCounter;
 	uint32 RevisionNumber = 0;
 	uint32 PreviousRevisionNumber = 0;
 
diff --git a/Engine/Source/Runtime/Engine/Public/SceneView.h b/Engine/Source/Runtime/Engine/Public/SceneView.h
--- a/Engine/Source/Runtime/Engine/Public/SceneView.h
+++ b/Engine/Source/Runtime/Engine/Public/SceneView.h
@@ -2110,8 +2110,7 @@ public:
 	uint64 FrameCounter = 0;
 
 	// When deleted, remove PRAGMA_DISABLE_DEPRECATION_WARNINGS / PRAGMA_ENABLE_DEPRECATION_WARNINGS from ~FSceneViewFamily()
-	/** If set, overrides FrameCounter value, useful for determistic Movie Render Queue behaviour. **/
-	UE_DEPRECATED(5.3, "Do not use, this is here as a temporary workaround for another issue.")
+	UE_DEPRECATED(5.3, "Do not use, has no effect, and will be removed in a future release.")
 	TOptional<uint64> OverrideFrameCounter;
 
 	/** Indicates this view family is an additional one. */
diff --git a/Engine/Source/Runtime/Launch/Private/LaunchEngineLoop.cpp b/Engine/Source/Runtime/Launch/Private/LaunchEngineLoop.cpp
--- a/Engine/Source/Runtime/Launch/Private/LaunchEngineLoop.cpp
+++ b/Engine/Source/Runtime/Launch/Private/LaunchEngineLoop.cpp
@@ -2828,8 +2828,6 @@ int32 FEngineLoop::PreInitPreStartupScreen(const TCHAR* CmdLine)
 	}
 
 #if WITH_ENGINE && CSV_PROFILER
-	FCoreDelegates::OnBeginFrame.AddStatic(UpdateCoreCsvStats_BeginFrame);
-	FCoreDelegates::OnEndFrame.AddStatic(UpdateCoreCsvStats_EndFrame);
 	FCsvProfiler::Get()->Init();
 #endif
 
@@ -5617,6 +5615,10 @@ void FEngineLoop::Tick()
 			}
 		#endif
 
+#if WITH_ENGINE && CSV_PROFILER
+		UpdateCoreCsvStats_BeginFrame();
+#endif
+
 		FCoreDelegates::OnBeginFrame.Broadcast();
 
 		// flush debug output which has been buffered by other threads
@@ -5991,17 +5993,8 @@ void FEngineLoop::Tick()
 		// If multithreaded rendering is off, it can cause a bad ordering of game and rendering markers.
 		GEngine->SetSimulationLatencyMarkerEnd(CurrentFrameCounter);
 
-		// Increment global frame counter. Once for each engine tick.
-		GFrameCounter++;
-
-		ENQUEUE_RENDER_COMMAND(FrameCounter)(
-			[CurrentFrameCounter = GFrameCounter](FRHICommandListImmediate& RHICmdList)
-		{
-			GFrameCounterRenderThread = CurrentFrameCounter;
-		});
-
 		// Disregard first few ticks for total tick time as it includes loading and such.
-		if (GFrameCounter > 6)
+		if (GFrameCounter > 5)
 		{
 			TotalTickTime += FApp::GetDeltaTime();
 		}
@@ -6049,6 +6042,23 @@ void FEngineLoop::Tick()
 
 		FCoreDelegates::OnEndFrame.Broadcast();
 
+		// Increment global frame counter. Once for each engine tick.
+		GFrameCounter++;
+
+		ENQUEUE_RENDER_COMMAND(FrameCounter)(
+			[CurrentFrameCounter = GFrameCounter](FRHICommandListImmediate& RHICmdList)
+		{
+			GFrameCounterRenderThread = CurrentFrameCounter;
+		});
+
+#if WITH_ENGINE && CSV_PROFILER
+		// By design, update this after incrementing GFrameCounter.  Calls PlatformMemoryHelpers::GetFrameMemoryStats, which caches
+		// memory stats based on the value of GFrameCounter.  Placing this after the increment guarantees that cached memory stats are
+		// canonically updated at this same point each frame, as opposed to arbitrarily in various other code paths (at least when the
+		// CSV profiler is compiled in, which is true in all builds except shipping).
+		UpdateCoreCsvStats_EndFrame();
+#endif
+
 		// Tick DumpGPU to start/stop frame dumps
 		#if WITH_ENGINE && WITH_DUMPGPU
 			UE::RenderCore::DumpGPU::TickEndFrame();
diff --git a/Engine/Source/Runtime/Renderer/Private/SceneRendering.cpp b/Engine/Source/Runtime/Renderer/Private/SceneRendering.cpp
--- a/Engine/Source/Runtime/Renderer/Private/SceneRendering.cpp
+++ b/Engine/Source/Runtime/Renderer/Private/SceneRendering.cpp
@@ -4961,10 +4961,7 @@ void FRendererModule::BeginRenderingViewFamilies(FCanvas* Canvas, TArrayView<FSc
 
 	for (FSceneViewFamily* ViewFamily : ViewFamilies)
 	{
-		PRAGMA_DISABLE_DEPRECATION_WARNINGS
-		// If optional parameter OverrideFrameCounter is set, use its value instead of GFrameCounter.
-		ViewFamily->FrameCounter = ViewFamily->OverrideFrameCounter.IsSet() ? ViewFamily->OverrideFrameCounter.GetValue() : GFrameCounter;
-		PRAGMA_ENABLE_DEPRECATION_WARNINGS
+		ViewFamily->FrameCounter = GFrameCounter;
 
 #if !(UE_BUILD_SHIPPING || UE_BUILD_TEST)
 		{
