fix the EMDoneTrigger timeout failure in SuspendApp_EventThread()
fix issue introduced recently regarding swapping W/H in portrait mode.
#android

diff --git a/Engine/Source/Runtime/ApplicationCore/Private/Android/AndroidWindow.cpp b/Engine/Source/Runtime/ApplicationCore/Private/Android/AndroidWindow.cpp
--- a/Engine/Source/Runtime/ApplicationCore/Private/Android/AndroidWindow.cpp
+++ b/Engine/Source/Runtime/ApplicationCore/Private/Android/AndroidWindow.cpp
@@ -645,14 +645,6 @@ FPlatformRect FAndroidWindow::GetScreenRect(bool bUseEventThreadWindow)
 		bool bIsPortrait = GDeviceScreenOrientation == EDeviceScreenOrientation::Portrait || GDeviceScreenOrientation == EDeviceScreenOrientation::PortraitUpsideDown;
 		UE_LOG(LogAndroid, Log, TEXT("FAndroidWindow::GetScreenRect bIsPortrait=%d"), bIsPortrait);
 
-		if (bIsPortrait)
-		{
-			UE_LOG(LogAndroid, Log, TEXT("FAndroidWindow::GetScreenRect(swap WH) bIsPortrait=%d"), bIsPortrait);
-			int temp = ScreenWidth;
-			ScreenWidth = ScreenHeight;
-			ScreenHeight = temp;
-		}
-
 		// save for future calls
 		CurrentParams.WindowWidth = ScreenWidth;
 		CurrentParams.WindowHeight = ScreenHeight;
diff --git a/Engine/Source/Runtime/Launch/Private/Android/AndroidEventManager.cpp b/Engine/Source/Runtime/Launch/Private/Android/AndroidEventManager.cpp
--- a/Engine/Source/Runtime/Launch/Private/Android/AndroidEventManager.cpp
+++ b/Engine/Source/Runtime/Launch/Private/Android/AndroidEventManager.cpp
@@ -262,9 +262,7 @@ void FAppEventManager::HandleWindowCreated_EventThread(void* InWindow)
 	}
 	EnqueueAppEvent(APP_EVENT_STATE_WINDOW_CREATED, FAppEventData((ANativeWindow*)InWindow));
 
-#if USE_ANDROID_STANDALONE
 	FApp::SetHasFocusFunction(&IsAppPausedOrSuspended);
-#endif
 }
 
 void FAppEventManager::HandleWindowClosed_EventThread()
@@ -281,9 +279,7 @@ void FAppEventManager::HandleWindowClosed_EventThread()
 
 	EnqueueAppEvent(APP_EVENT_STATE_WINDOW_DESTROYED);
 
-#if USE_ANDROID_STANDALONE
 	FApp::SetHasFocusFunction(nullptr);
-#endif
 }
 
 
@@ -409,11 +405,7 @@ bool FAppEventManager::IsGamePaused()
 
 bool FAppEventManager::IsGameInFocus()
 {
-#if USE_ANDROID_STANDALONE
 	return (bWindowInFocus && bHaveWindow && bHaveGame);
-#else
-	return (bWindowInFocus && bHaveWindow);
-#endif
 }
 
 
diff --git a/Engine/Source/Runtime/Launch/Private/Android/LaunchAndroid.cpp b/Engine/Source/Runtime/Launch/Private/Android/LaunchAndroid.cpp
--- a/Engine/Source/Runtime/Launch/Private/Android/LaunchAndroid.cpp
+++ b/Engine/Source/Runtime/Launch/Private/Android/LaunchAndroid.cpp
@@ -1649,7 +1649,7 @@ static void SuspendApp_EventThread()
 	{
 		FPlatformProcess::ReturnSynchEventToPool(EventToDelete);
 	});
-	EMDoneTrigger->Reset();
+	
 
 	// perform the delegates before the window handle is cleared.
 	// This ensures any tasks that require a window handle will have it before we block the RT on the invalid window.
@@ -1661,6 +1661,7 @@ static void SuspendApp_EventThread()
 		FCoreDelegates::ApplicationWillEnterBackgroundDelegate.Broadcast();
 		FAppEventManager::GetInstance()->PauseAudio();
 		FAppEventManager::ReleaseMicrophone(false);
+
 		EMDoneTrigger->Trigger();
 	}));
 
@@ -1674,16 +1675,19 @@ static void SuspendApp_EventThread()
 
 	// wait for a period of time before blocking rendering
 	UE_LOG(LogAndroid, Log, TEXT("SuspendApp_EventThread -> , waiting for event manager to process. tid: %d"), FPlatformTLS::GetCurrentThreadId());
-
+#if USE_ANDROID_STANDALONE
+	EMDoneTrigger->Reset();
+#endif
 	bool bSuccess = EMDoneTrigger->Wait(4000);
 	float ElapsedTimeInMs_EMDoneTrigger_Wait = FPlatformTime::ToMilliseconds(FPlatformTime::Cycles() - StartCycles);
-	UE_CLOG(!bSuccess, LogAndroid, Log, TEXT("SuspendApp_EventThread -> backgrounding callback, not responded in timely manner."));
-	STANDALONE_DEBUG_LOGf(LogAndroid, TEXT("SuspendApp_EventThread -> EMDoneTrigger->Wait, waited '%f' ms"), (float)ElapsedTimeInMs_EMDoneTrigger_Wait);
+	UE_CLOG(!bSuccess, LogAndroid, Log, TEXT("SuspendApp_EventThread -> ERROR: backgrounding callback, not responded in timely manner. EMDoneTrigger->Wait, waited '%f' ms"), (float)ElapsedTimeInMs_EMDoneTrigger_Wait);
+	UE_LOG(LogAndroid, Log, TEXT("SuspendApp_EventThread -> EMDoneTrigger->Wait, waited '%f' ms"), (float)ElapsedTimeInMs_EMDoneTrigger_Wait);
 
 	BlockRendering();
 
 	// Suspend the GT.
 	FAppEventManager::GetInstance()->EnqueueAppEvent(APP_EVENT_STATE_APP_SUSPENDED);
+	UE_LOG(LogAndroid, Log, TEXT("SuspendApp_EventThread(EOF)"));
 }
 
 //Called from the event process thread
diff --git a/Engine/Source/Runtime/OpenGLDrv/Private/Android/AndroidEGL.cpp b/Engine/Source/Runtime/OpenGLDrv/Private/Android/AndroidEGL.cpp
--- a/Engine/Source/Runtime/OpenGLDrv/Private/Android/AndroidEGL.cpp
+++ b/Engine/Source/Runtime/OpenGLDrv/Private/Android/AndroidEGL.cpp
@@ -1333,7 +1333,7 @@ void BlockRendering()
 #endif
 	check(GIsRHIInitialized);
 
-	UE_LOG(LogAndroid, Log, TEXT("Blocking renderer on invalid window."));
+	UE_LOG(LogAndroid, Log, TEXT("Blocking renderer on suspended window."));
 	
 	TSharedPtr<FEvent, ESPMode::ThreadSafe> BlockedTrigger = MakeShareable(FPlatformProcess::GetSynchEventFromPool(), [](FEvent* EventToDelete)
 	{
diff --git a/Engine/Source/Runtime/OpenGLDrv/Private/Android/AndroidOpenGLFramePacer.cpp b/Engine/Source/Runtime/OpenGLDrv/Private/Android/AndroidOpenGLFramePacer.cpp
--- a/Engine/Source/Runtime/OpenGLDrv/Private/Android/AndroidOpenGLFramePacer.cpp
+++ b/Engine/Source/Runtime/OpenGLDrv/Private/Android/AndroidOpenGLFramePacer.cpp
@@ -190,7 +190,7 @@ bool FAndroidOpenGLFramePacer::SupportsFramePaceInternal(int32 QueryFramePace, i
 		{
 			RefreshRatesString += FString::Printf(TEXT(" %d"), Rate);
 		}
-		UE_LOG(LogRHI, Log, TEXT("Supported Refresh Rates:%s"), *RefreshRatesString);
+		UE_LOG(LogRHI, Log, TEXT("FAndroidOpenGLFramePacer -> Supported Refresh Rates:%s"), *RefreshRatesString);
 
 		for (int32 Rate : RefreshRates)
 		{
diff --git a/Engine/Source/Runtime/VulkanRHI/Private/Android/VulkanAndroidPlatform.cpp b/Engine/Source/Runtime/VulkanRHI/Private/Android/VulkanAndroidPlatform.cpp
--- a/Engine/Source/Runtime/VulkanRHI/Private/Android/VulkanAndroidPlatform.cpp
+++ b/Engine/Source/Runtime/VulkanRHI/Private/Android/VulkanAndroidPlatform.cpp
@@ -760,7 +760,7 @@ bool FAndroidVulkanFramePacer::SupportsFramePaceInternal(int32 QueryFramePace, i
 	TArray<int32> RefreshRates = FAndroidMisc::GetSupportedNativeDisplayRefreshRates();
 	RefreshRates.Sort();
 
-	FString DebugString = TEXT("Supported Refresh Rates:");
+	FString DebugString = TEXT("FAndroidVulkanFramePacer -> Supported Refresh Rates:");
 	for (int32 RefreshRate : RefreshRates)
 	{
 		DebugString += FString::Printf(TEXT(" %d"), RefreshRate);
