Fixed verse build phase to happen after new UObjects have been registered in the engine.  Wrong delegate was being used causing the build to happen prior to UObject definitions in some instances.

diff --git a/Engine/Source/Runtime/CoreUObject/Private/UObject/Obj.cpp b/Engine/Source/Runtime/CoreUObject/Private/UObject/Obj.cpp
--- a/Engine/Source/Runtime/CoreUObject/Private/UObject/Obj.cpp
+++ b/Engine/Source/Runtime/CoreUObject/Private/UObject/Obj.cpp
@@ -5281,11 +5281,9 @@ void InitUObject()
 
 	FCoreDelegates::OnShutdownAfterError.AddStatic(StaticShutdownAfterError);
 	FCoreDelegates::OnExit.AddStatic(StaticExit);
-#if !USE_PER_MODULE_UOBJECT_BOOTSTRAP && !IS_MONOLITHIC
 
 	// Always register the UObjects callback for VNI and general consistency with the callbacks ProcessNewlyLoadedUObjects calls.
 	RegisterProcessNewlyLoadedUObjects();
-#endif
 
 	struct Local
 	{
diff --git a/Engine/Source/Runtime/Launch/Private/LaunchEngineLoop.cpp b/Engine/Source/Runtime/Launch/Private/LaunchEngineLoop.cpp
--- a/Engine/Source/Runtime/Launch/Private/LaunchEngineLoop.cpp
+++ b/Engine/Source/Runtime/Launch/Private/LaunchEngineLoop.cpp
@@ -4398,10 +4398,8 @@ bool FEngineLoop::LoadCoreModules()
 {
 	// Always attempt to load CoreUObject. It requires additional pre-init which is called from its module's StartupModule method.
 #if WITH_COREUOBJECT
-#if USE_PER_MODULE_UOBJECT_BOOTSTRAP // otherwise do it later
 	// Always register the UObjects callback for VNI and general consistency with the callbacks ProcessNewlyLoadedUObjects calls.
 	RegisterProcessNewlyLoadedUObjects();
-#endif
 	return FModuleManager::Get().LoadModule(TEXT("CoreUObject")) != nullptr;
 #else
 	return true;
