Revert some changes to Android event flow

#android

diff --git a/Engine/Source/Runtime/Launch/Private/Android/AndroidEventManager.cpp b/Engine/Source/Runtime/Launch/Private/Android/AndroidEventManager.cpp
--- a/Engine/Source/Runtime/Launch/Private/Android/AndroidEventManager.cpp
+++ b/Engine/Source/Runtime/Launch/Private/Android/AndroidEventManager.cpp
@@ -261,8 +261,6 @@ void FAppEventManager::HandleWindowCreated_EventThread(void* InWindow)
 		FirstInitialized = true;
 	}
 	EnqueueAppEvent(APP_EVENT_STATE_WINDOW_CREATED, FAppEventData((ANativeWindow*)InWindow));
-
-	FApp::SetHasFocusFunction(&IsAppPausedOrSuspended);
 }
 
 void FAppEventManager::HandleWindowClosed_EventThread()
@@ -278,8 +276,6 @@ void FAppEventManager::HandleWindowClosed_EventThread()
 	FAndroidWindow::ReleaseWindowRef((ANativeWindow*)ActiveWindow);
 
 	EnqueueAppEvent(APP_EVENT_STATE_WINDOW_DESTROYED);
-
-	FApp::SetHasFocusFunction(nullptr);
 }
 
 
@@ -405,7 +401,7 @@ bool FAppEventManager::IsGamePaused()
 
 bool FAppEventManager::IsGameInFocus()
 {
-	return (bWindowInFocus && bHaveWindow && bHaveGame);
+	return (bWindowInFocus && bHaveWindow);
 }
 
 
