#UE Added an option to launch even though modules are missing or incompatible. This is not a recommended action, but it will allow devs to partially load plugins without needing to compile the code for them in order to inspect them.

diff --git a/Engine/Source/Runtime/Launch/Private/LaunchEngineLoop.cpp b/Engine/Source/Runtime/Launch/Private/LaunchEngineLoop.cpp
--- a/Engine/Source/Runtime/Launch/Private/LaunchEngineLoop.cpp
+++ b/Engine/Source/Runtime/Launch/Private/LaunchEngineLoop.cpp
@@ -6539,12 +6539,18 @@ bool FEngineLoop::AppInit( )
 
 				// Ask whether to compile before continuing
 				FString CompilePrompt = ModulesList + TEXT("\nWould you like to rebuild them now?");
-				if (FPlatformMisc::MessageBoxExt(EAppMsgType::YesNo, *CompilePrompt, *FString::Printf(TEXT("Missing %s Modules"), FApp::GetProjectName())) == EAppReturnType::No)
+				if (FPlatformMisc::MessageBoxExt(EAppMsgType::YesNo, *CompilePrompt, *FString::Printf(TEXT("Missing %s Modules"), FApp::GetProjectName())) == EAppReturnType::Yes)
 				{
-					return false;
+					bNeedCompile = true;
+				}
+				else
+				{
+					FString ProceedWithoutCompilePrompt = ModulesList + TEXT("\nWould you like to proceed anyway? Code and content that depends on these modules may not work correctly. Enter at your own risk.");
+					if (FPlatformMisc::MessageBoxExt(EAppMsgType::YesNo, *ProceedWithoutCompilePrompt, *FString::Printf(TEXT("Missing %s Modules"), FApp::GetProjectName())) == EAppReturnType::No)
+					{
+						return false;
+					}
 				}
-
-				bNeedCompile = true;
 			}
 			else if(EditorTargetFileName.Len() > 0 && !FPaths::FileExists(EditorTargetFileName))
 			{
