Remove the GIsGuarded changes so app will force exit on fatal error rather than try to continue.
#android

diff --git a/Engine/Source/Runtime/Launch/Private/Android/LaunchAndroid.cpp b/Engine/Source/Runtime/Launch/Private/Android/LaunchAndroid.cpp
--- a/Engine/Source/Runtime/Launch/Private/Android/LaunchAndroid.cpp
+++ b/Engine/Source/Runtime/Launch/Private/Android/LaunchAndroid.cpp
@@ -701,17 +701,13 @@ int32 AndroidMain(struct android_app* state)
 		}
 	}
 	// make sure GEngineLoop::Exit() is always called.
-	struct EngineLoopCleanupGuard
+	struct EngineLoopCleanup
 	{
-		~EngineLoopCleanupGuard()
+		~EngineLoopCleanup()
 		{
-			if (GIsGuarded)
-			{
-				GIsGuarded = false;
-			}
 			EngineExit();
 		}
-	} CleanupGuard;
+	} Cleanup;
 
 	// read the command line file
 	InitCommandLine();
@@ -728,49 +724,10 @@ int32 AndroidMain(struct android_app* state)
 	}
 #endif
 
-	STANDALONE_DEBUG_LOGf(LogAndroid, TEXT("LogLoad::AndroidMain - NDKVersion = %d"), PLATFORM_USED_NDK_VERSION_INTEGER);
-
-	bool bEnableGuarded = false;
-
-	FString reportcrashOverride;
-	if (FParse::Value(CmdLine, TEXT("reportcrashOverride="), reportcrashOverride))
-	{
-		GAlwaysReportCrash = FCString::Atoi(*reportcrashOverride);
-		STANDALONE_DEBUG_LOGf(LogAndroid, TEXT("LogLoad::AndroidMain cmdline arg(reportcrashOverride) DISABLE GAlwaysReportCrash. isDebuggerPresent = %d, GAlwaysReportCrash = %d, bEnableGuarded = %d, GIsGuarded = %d"), FPlatformMisc::IsDebuggerPresent(), GAlwaysReportCrash, bEnableGuarded, GIsGuarded);
-	}
-
-
-#if (UE_BUILD_DEBUG || UE_BUILD_DEVELOPMENT)
-	if (true && !GAlwaysReportCrash)
-#else
-	if (FPlatformMisc::IsDebuggerPresent() && !GAlwaysReportCrash)
-#endif
-	{
-		// Don't use exception handling when a debugger is attached to exactly trap the crash. This does NOT check
-		// whether we are the first instance or not!
-		bEnableGuarded = false;
-	}
-	else
-	{
-		bEnableGuarded = true;
-	}
-
-	if (bEnableGuarded && !GIsGuarded)
-	{
-		STANDALONE_DEBUG_LOGf(LogAndroid, TEXT("LogLoad::AndroidMain ENABLE GIsGuarded. isDebuggerPresent = %d, GAlwaysReportCrash = %d, bEnableGuarded = %d, GIsGuarded = %d"), FPlatformMisc::IsDebuggerPresent(), GAlwaysReportCrash, bEnableGuarded, GIsGuarded);
-
-		GIsGuarded = true;
-	}
-
-	STANDALONE_DEBUG_LOGf(LogAndroid, TEXT("LogLoad::AndroidMain (Guarded state) IsDebuggerPresent = %d, GAlwaysReportCrash = %d, bEnableGuarded = %d, GIsGuarded = %d"), FPlatformMisc::IsDebuggerPresent(), GAlwaysReportCrash, bEnableGuarded, GIsGuarded);
-
 	EventHandlerEvent = FPlatformProcess::GetSynchEventFromPool(false);
 	FPlatformMisc::LowLevelOutputDebugString(TEXT("Created sync event\n"));
 	FAppEventManager::GetInstance()->SetEventHandlerEvent(EventHandlerEvent);
 
-	STANDALONE_DEBUG_LOGf(LogAndroid, TEXT("LogLoad::AndroidMain (after SetEventHandlerEvent) IsDebuggerPresent = %d, GAlwaysReportCrash = %d, bEnableGuarded = %d, GIsGuarded = %d"), FPlatformMisc::IsDebuggerPresent(), GAlwaysReportCrash, bEnableGuarded, GIsGuarded);
-
-
 	// ready for onCreate to complete
 	GEventHandlerInitialized = true;
 
@@ -857,11 +814,7 @@ int32 AndroidMain(struct android_app* state)
 
 	BootTimingPoint("Tick loop starting");
 	DumpBootTiming();
-	if (bEnableGuarded && !GIsGuarded)
-	{
-		UE_LOG(LogLoad, Log, TEXT("(AndroidMain) ENABLE GIsGuarded before Engine Tick, IsDebuggerPresent = %d, GAlwaysReportCrash = %d, bEnableGuarded = %d, GIsGuarded = %d"), FPlatformMisc::IsDebuggerPresent(), GAlwaysReportCrash, bEnableGuarded, GIsGuarded);
-		GIsGuarded = true;
-	}
+
 	// tick until done
 	while (!IsEngineExitRequested())
 	{
