Fix: for portraid mode stretching issue on some devices. and ASIS cleanup
#android

diff --git a/Engine/Source/Runtime/ApplicationCore/Private/Android/AndroidWindow.cpp b/Engine/Source/Runtime/ApplicationCore/Private/Android/AndroidWindow.cpp
--- a/Engine/Source/Runtime/ApplicationCore/Private/Android/AndroidWindow.cpp
+++ b/Engine/Source/Runtime/ApplicationCore/Private/Android/AndroidWindow.cpp
@@ -195,46 +195,28 @@ JNI_METHOD void Java_com_epicgames_unreal_GameActivity_nativeSetSafezoneInfo(JNI
 
 #if USE_ANDROID_STANDALONE
 
-extern void GAndroidWindowLock_Lock(FString calledBy);
-extern void GAndroidWindowLock_Unlock(FString calledBy);
-
 JNI_METHOD void Java_com_epicgames_makeaar_GameActivityForMakeAAR_nativeSetSurfaceOverride(JNIEnv* jenv, jobject thiz, jobject surface, jint x, jint y)
 {
+	ANativeWindow* prev = (ANativeWindow*)GAndroidWindowOverride;
 	if (surface != 0)
 	{
 		GSurfaceViewX = x;
 		GSurfaceViewY = y;
-		void* prev = GAndroidWindowOverride;
-		GAndroidWindowOverride = (ANativeWindow*)ANativeWindow_fromSurface(jenv, surface);
-		STANDALONE_DEBUG_LOG(TEXT("nativeSetSurfaceOverride(makeaar) applied: prev to new %p -> %p, (%d, %d)"), prev, GAndroidWindowOverride, GSurfaceViewX, GSurfaceViewY);
-		if (prev != GAndroidWindowOverride)
-		{
-			//GAndroidWindowLock_Unlock("Java_com_epicgames_makeaar_GameActivityForMakeAAR_nativeSetSurfaceOverride");
-		}
 
-		//if (GAndroidWindowOverride != nullptr)
-		//{
-		//	STANDALONE_DEBUG_LOGf(LogAndroid, TEXT("nativeSetSurfaceOverride. FTaskGraphInterface::IsRunning() = %d, GAndroidWindowOverride=%x"), FTaskGraphInterface::IsRunning(), GAndroidWindowOverride);
-
-		//	if (FTaskGraphInterface::IsRunning())
-		//	{
-		//		FSimpleDelegateGraphTask::CreateAndDispatchWhenReady(
-		//			FSimpleDelegateGraphTask::FDelegate::CreateLambda([=]()
-		//				{
-		//					STANDALONE_DEBUG_LOGf(LogAndroid, TEXT("from nativeSetSurfaceOverride, trigger UnlockAndroidWindow. FTaskGraphInterface::IsRunning() = %d, GAndroidWindowOverride=%x"), FTaskGraphInterface::IsRunning(), GAndroidWindowOverride);
-
-		//					FPlatformMisc::UnlockAndroidWindow();
-		//				}), TStatId(), nullptr, ENamedThreads::GameThread);
-		//	}
-		//}
+		GAndroidWindowOverride = (ANativeWindow*)ANativeWindow_fromSurface(jenv, surface);
+		UE_LOG(LogAndroid, Log, TEXT("nativeSetSurfaceOverride applied: prev to new %p -> %p, pos(%d, %d)"), prev, GAndroidWindowOverride, GSurfaceViewX, GSurfaceViewY);
 
 	}
 	else
 	{
 		GAndroidWindowOverride = nullptr;
 		
-		STANDALONE_DEBUG_LOG(TEXT("nativeSetSurfaceOverride(makeaar) setting to null and lock window"));
-		//GAndroidWindowLock_Lock("Java_com_epicgames_makeaar_GameActivityForMakeAAR_nativeSetSurfaceOverride");
+		STANDALONE_DEBUG_LOG(TEXT("nativeSetSurfaceOverride(makeaar) setting to null"));
+	}
+
+	if (prev != nullptr)
+	{
+		ANativeWindow_release(prev);
 	}
 }
 
@@ -269,6 +251,7 @@ void FAndroidWindow::InvalidateCachedScreenRect()
 void FAndroidWindow::AcquireWindowRef(ANativeWindow* InWindow)
 {
 #if USE_ANDROID_JNI
+#if USE_ANDROID_STANDALONE
 	STANDALONE_DEBUG_LOG(TEXT("AcquireWindowRef USE_ANDROID_JNI is enabled: InWindow=%p, GAcquiredWindow=%p, GAndroidWindowOverride=%p"), InWindow, GAcquiredWindow, GAndroidWindowOverride);
 
 	if (InWindow == nullptr)
@@ -293,16 +276,14 @@ void FAndroidWindow::AcquireWindowRef(ANativeWindow* InWindow)
 	check(GAcquiredWindow == NULL);
 	STANDALONE_DEBUG_LOG(TEXT("AcquireWindowRef USE_ANDROID_JNI is enabled: %p and ANativeWindow_acquire"), InWindow);
 
-	ANativeWindow_acquire(InWindow);
-
 	// Added logic to store the Acquired window and when calling ReleaseWindowRef, check if the window is the same and only release if it matches
 	// This logic is to deal with the fact Android lifecycles for activities can overlap. ideally we would create a context based container to manage this
 	// but for now this is a useful protection.
 	GAcquiredWindow = InWindow;
 	STANDALONE_DEBUG_LOG(TEXT("FAndroidWindow::AcquireWindowRef GAcquiredWindow=%p, GAndroidWindowOverride=%p"), GAcquiredWindow, GAndroidWindowOverride);
+#endif //USE_ANDROID_STANDALONE
 
-#else
-	STANDALONE_DEBUG_LOG(TEXT("AcquireWindowRef USE_ANDROID_JNI is NOT enabled: %p"), InWindow);
+	ANativeWindow_acquire(InWindow);
 
 #endif
 }
@@ -310,6 +291,7 @@ void FAndroidWindow::AcquireWindowRef(ANativeWindow* InWindow)
 void FAndroidWindow::ReleaseWindowRef(ANativeWindow* InWindow)
 {
 #if USE_ANDROID_JNI
+#if USE_ANDROID_STANDALONE
 	if (GAcquiredWindow == nullptr && InWindow == nullptr)
 	{
 		STANDALONE_DEBUG_LOG(TEXT("ReleaseWindowRef skipped because GAcquiredWindow is null.  Window %p reference will not be released."), InWindow);
@@ -324,7 +306,7 @@ void FAndroidWindow::ReleaseWindowRef(ANativeWindow* InWindow)
 	}
 
 	STANDALONE_DEBUG_LOG(TEXT("ReleaseWindowRef using window: %p"), InWindow);
-
+#endif //USE_ANDROID_STANDALONE
 	ANativeWindow_release(InWindow);
 #endif
 }
@@ -376,7 +358,7 @@ bool FAndroidWindow::WaitForWindowDimensions()
 	{
 		if (IsEngineExitRequested()
 #if USE_ANDROID_EVENTS
-			|| FAppEventManager::GetInstance()->WaitForEventInQueue(EAppEventState::APP_EVENT_STATE_ON_DESTROY, 0.0f)
+		|| FAppEventManager::GetInstance()->WaitForEventInQueue(EAppEventState::APP_EVENT_STATE_ON_DESTROY, 0.0f)
 #endif
 			)
 		{
@@ -566,7 +548,7 @@ FPlatformRect FAndroidWindow::GetScreenRect(bool bUseEventThreadWindow)
 	// too much of the following code needs JNI things, just assume override
 #if !USE_ANDROID_JNI
 
-	UE_LOG(LogAndroid, Fatal, TEXT("FAndroidWindow::CalculateSurfaceSize currently expedcts non-JNI platforms to override resolution"));
+	UE_LOG(LogAndroid, Fatal, TEXT("FAndroidWindow::CalculateSurfaceSize currently expects non-JNI platforms to override resolution"));
 	return FPlatformRect();
 #else
 
@@ -641,10 +623,6 @@ FPlatformRect FAndroidWindow::GetScreenRect(bool bUseEventThreadWindow)
 			}
 		}
 
-
-		bool bIsPortrait = GDeviceScreenOrientation == EDeviceScreenOrientation::Portrait || GDeviceScreenOrientation == EDeviceScreenOrientation::PortraitUpsideDown;
-		UE_LOG(LogAndroid, Log, TEXT("FAndroidWindow::GetScreenRect bIsPortrait=%d"), bIsPortrait);
-
 		// save for future calls
 		CurrentParams.WindowWidth = ScreenWidth;
 		CurrentParams.WindowHeight = ScreenHeight;
@@ -708,7 +686,13 @@ void FAndroidWindow::CalculateSurfaceSize(int32_t& SurfaceWidth, int32_t& Surfac
 	// do not convert to a surface size that is larger than native resolution
 	// Mobile VR doesn't need buffer quantization as Unreal never renders directly to the buffer in VR mode. 
 	static const bool bIsMobileVRApp = AndroidThunkCpp_IsOculusMobileApplication();
+	
+#if USE_ANDROID_STANDALONE
+	const int DividableBy = 1;	// don't change size of external window 
+#else
 	const int DividableBy = bIsMobileVRApp ? 1 : 8;
+#endif
+
 	SurfaceWidth = (SurfaceWidth / DividableBy) * DividableBy;
 	SurfaceHeight = (SurfaceHeight / DividableBy) * DividableBy;
 #endif
diff --git a/Engine/Source/Runtime/Launch/Private/Android/AndroidEventManager.cpp b/Engine/Source/Runtime/Launch/Private/Android/AndroidEventManager.cpp
--- a/Engine/Source/Runtime/Launch/Private/Android/AndroidEventManager.cpp
+++ b/Engine/Source/Runtime/Launch/Private/Android/AndroidEventManager.cpp
@@ -97,7 +97,7 @@ void FAppEventManager::Tick()
 			break;
 		case APP_EVENT_STATE_WINDOW_DESTROYED:
 			bHaveWindow = false;
-			STANDALONE_DEBUG_LOGf(LogAndroidEvents, TEXT("APP_EVENT_STATE_WINDOW_DESTROYED, %d, %d, %d"), int(bRunning), int(bHaveWindow), int(bHaveGame));
+			STANDALONE_DEBUG_LOGf(LogAndroidEvents, TEXT("APP_EVENT_STATE_WINDOW_DESTROYED, %d, %d, %d"), int(bRunning), int(bWindowInFocus), int(bHaveGame));
 			break;
 		case APP_EVENT_STATE_ON_START:
 			//doing nothing here
@@ -114,9 +114,11 @@ void FAppEventManager::Tick()
 		case APP_EVENT_STATE_ON_PAUSE:
 			FAndroidAppEntry::OnPauseEvent();
 			bHaveGame = false;
+			STANDALONE_DEBUG_LOGf(LogAndroidEvents, TEXT("APP_EVENT_STATE_ON_PAUSE, %d, %d, %d"), int(bRunning), int(bWindowInFocus), int(bHaveGame));
 			break;
 		case APP_EVENT_STATE_ON_RESUME:
 			bHaveGame = true;
+			STANDALONE_DEBUG_LOGf(LogAndroidEvents, TEXT("APP_EVENT_STATE_ON_RESUME, %d, %d, %d"), int(bRunning), int(bWindowInFocus), int(bHaveGame));
 			break;
 
 		// window focus events that follow their own hierarchy, and might or might not respect App main events hierarchy
@@ -213,12 +215,16 @@ FAppEventManager::FAppEventManager():
 	CVarScale->SetOnChangedCallback(FConsoleVariableDelegate::CreateStatic(&FAppEventManager::OnScaleFactorChanged));
 
 	IConsoleVariable* CVarResX = IConsoleManager::Get().FindConsoleVariable(TEXT("r.Mobile.DesiredResX"));
-	check(CVarResX);
-	CVarResX->SetOnChangedCallback(FConsoleVariableDelegate::CreateStatic(&FAppEventManager::OnScaleFactorChanged));
+	if (CVarResX != nullptr)
+	{
+		CVarResX->SetOnChangedCallback(FConsoleVariableDelegate::CreateStatic(&FAppEventManager::OnScaleFactorChanged));
+	}
 
 	IConsoleVariable* CVarResY = IConsoleManager::Get().FindConsoleVariable(TEXT("r.Mobile.DesiredResY"));
-	check(CVarResY);
-	CVarResY->SetOnChangedCallback(FConsoleVariableDelegate::CreateStatic(&FAppEventManager::OnScaleFactorChanged));
+	if (CVarResY != nullptr)
+	{
+		CVarResY->SetOnChangedCallback(FConsoleVariableDelegate::CreateStatic(&FAppEventManager::OnScaleFactorChanged));
+	}
 }
 
 void FAppEventManager::OnScaleFactorChanged(IConsoleVariable* CVar)
@@ -234,16 +240,11 @@ void FAppEventManager::OnScaleFactorChanged(IConsoleVariable* CVar)
     }
 }
 
-static bool IsAppPausedOrSuspended()
-{
-	return FAppEventManager::GetInstance()->IsGamePaused() && !FAppEventManager::GetInstance()->IsGameInFocus();
-}
-
 void FAppEventManager::HandleWindowCreated_EventThread(void* InWindow)
 {
 	if (InWindow == nullptr)
 	{
-		STANDALONE_DEBUG_LOGf(LogAndroidEvents, TEXT("FAppEventManager::HandleWindowCreated_EventThread was given a NULL window so IGNORE."));
+		FPlatformMisc::LowLevelOutputDebugStringf(TEXT("FAppEventManager::HandleWindowCreated_EventThread was given a NULL window so IGNORE."));
 		return;
 	}
 
@@ -252,7 +253,6 @@ void FAppEventManager::HandleWindowCreated_EventThread(void* InWindow)
 	// Make sure window will not be deleted until event is processed
 	// Window could be deleted by OS while event queue stuck at game start-up phase
 	FAndroidWindow::AcquireWindowRef((ANativeWindow*)InWindow);
-
 	FAndroidWindow::SetHardwareWindow_EventThread(InWindow);
 
 	if (!AlreadyInited)
@@ -272,8 +272,8 @@ void FAppEventManager::HandleWindowClosed_EventThread()
 		return;
 	}
 
-	FAndroidWindow::SetHardwareWindow_EventThread(nullptr);
 	FAndroidWindow::ReleaseWindowRef((ANativeWindow*)ActiveWindow);
+	FAndroidWindow::SetHardwareWindow_EventThread(nullptr);
 
 	EnqueueAppEvent(APP_EVENT_STATE_WINDOW_DESTROYED);
 }
@@ -401,7 +401,11 @@ bool FAppEventManager::IsGamePaused()
 
 bool FAppEventManager::IsGameInFocus()
 {
+#if USE_ANDROID_STANDALONE
+	return (bWindowInFocus && bHaveWindow && bHaveGame);
+#else
 	return (bWindowInFocus && bHaveWindow);
+#endif
 }
 
 
diff --git a/Engine/Source/Runtime/Launch/Private/Android/LaunchAndroid.cpp b/Engine/Source/Runtime/Launch/Private/Android/LaunchAndroid.cpp
--- a/Engine/Source/Runtime/Launch/Private/Android/LaunchAndroid.cpp
+++ b/Engine/Source/Runtime/Launch/Private/Android/LaunchAndroid.cpp
@@ -291,86 +291,11 @@ static bool bShouldRestartFromInterrupt = false;
 static bool bIgnorePauseOnDownloaderStart = false;
 bool bReadyToProcessEvents = false;
 
-#if USE_ANDROID_STANDALONE
-static int GAndroidWindowLockRefCount = 0;
-static const bool bForce_GAndroidWindowLock = true;
-#endif
-
 pthread_t G_AndroidEventThread;
 struct android_app* GNativeAndroidApp = NULL;
+static uint32 EventThreadID = 0;
 
-void GAndroidWindowLock_Lock(FString calledBy)
-{
-#if !USE_ANDROID_STANDALONE
-	GAndroidWindowLock.Lock();
-#else
-	bool doAction = bForce_GAndroidWindowLock;
-	if (GAndroidWindowLockRefCount >= 0)
-	{
-		++GAndroidWindowLockRefCount;
-		if (GAndroidWindowLockRefCount == 1)
-		{
-			doAction = true;
-		}
-		else
-		{
-			STANDALONE_DEBUG_LOGf(LogAndroid, TEXT("WARNING GAndroidWindowLock_Lock(%s). ALREADY LOCKED. GAndroidWindowLockRefCount=%d, bAppIsActive_EventThread=%d"), *calledBy, GAndroidWindowLockRefCount, bAppIsActive_EventThread);
-		}
-	}
-	else
-	{
-		STANDALONE_DEBUG_LOGf(LogAndroid, TEXT("GAndroidWindowLock_Lock(%s). ERROR: GAndroidWindowLockRefCount=%d, bAppIsActive_EventThread=%d"), *calledBy, GAndroidWindowLockRefCount, bAppIsActive_EventThread);
-	}
-	if (doAction)
-	{
-		GAndroidWindowLock.Lock();
-
-		STANDALONE_DEBUG_LOGf(LogAndroid, TEXT("GAndroidWindowLock_Lock(%s). LOCKED. GAndroidWindowLockRefCount=%d, bAppIsActive_EventThread=%d"), *calledBy, GAndroidWindowLockRefCount, bAppIsActive_EventThread);
-	}
-	else
-	{
-		STANDALONE_DEBUG_LOGf(LogAndroid, TEXT("WARNING GAndroidWindowLock_Lock(%s). NOOP: GAndroidWindowLockRefCount=%d, bAppIsActive_EventThread=%d bReadyToProcessEvents=%d"), *calledBy, GAndroidWindowLockRefCount, bAppIsActive_EventThread, bReadyToProcessEvents);
-	}
-#endif
-}
 
-void GAndroidWindowLock_Unlock(FString calledBy)
-{
-#if !USE_ANDROID_STANDALONE
-	GAndroidWindowLock.Unlock();
-#else
-	bool doAction = bForce_GAndroidWindowLock;
-	if (GAndroidWindowLockRefCount > 0)
-	{
-		if (--GAndroidWindowLockRefCount == 0)
-		{
-			doAction = true;
-		}
-		else
-		{
-			STANDALONE_DEBUG_LOGf(LogAndroid, TEXT("WARNING GAndroidWindowLock_Unlock(%s). KEEP LOCKED, GAndroidWindowLockRefCount=%d, bAppIsActive_EventThread=%d"), *calledBy, GAndroidWindowLockRefCount, bAppIsActive_EventThread);
-		}
-		if (GAndroidWindowLockRefCount < 0)
-		{
-			STANDALONE_DEBUG_LOGf(LogAndroid, TEXT("GAndroidWindowLock_Unlock(%s). ERROR UNLOCK MISMATCH, GAndroidWindowLockRefCount=%d, bAppIsActive_EventThread=%d"), *calledBy, GAndroidWindowLockRefCount, bAppIsActive_EventThread);
-			GAndroidWindowLockRefCount = 0;
-		}
-	}
-	else
-	{
-		STANDALONE_DEBUG_LOGf(LogAndroid, TEXT("GAndroidWindowLock_Unlock(%s). SKIPPED, GAndroidWindowLockRefCount=%d, bAppIsActive_EventThread=%d"), *calledBy, GAndroidWindowLockRefCount, bAppIsActive_EventThread);
-	}
-	if (doAction)
-	{
-		GAndroidWindowLock.Unlock();
-		STANDALONE_DEBUG_LOGf(LogAndroid, TEXT("GAndroidWindowLock_Unlock(%s). UNLOCK, GAndroidWindowLockRefCount=%d, bAppIsActive_EventThread=%d"), *calledBy, GAndroidWindowLockRefCount, bAppIsActive_EventThread);
-	}
-	else
-	{
-		STANDALONE_DEBUG_LOGf(LogAndroid, TEXT("GAndroidWindowLock_Unlock(%s). NOOP: GAndroidWindowLockRefCount=%d, bAppIsActive_EventThread=%d bReadyToProcessEvents=%d"), *calledBy, GAndroidWindowLockRefCount, bAppIsActive_EventThread, bReadyToProcessEvents);
-	}
-#endif
-}
 static void SuspendApp_EventThread();
 static void ActivateApp_EventThread();
 #if USE_ANDROID_STANDALONE && !(UE_BUILD_SHIPPING || UE_BUILD_TEST)
@@ -387,7 +312,7 @@ void FPlatformMisc::UnlockAndroidWindow()
 	check(FTaskGraphInterface::IsRunning());
 	
 	bReadyToProcessEvents = true;
-	GAndroidWindowLock_Unlock("UnlockAndroidWindow");
+	GAndroidWindowLock.Unlock();
 }
 
 JNI_METHOD void Java_com_epicgames_unreal_GameActivity_nativeResumeMainInit(JNIEnv* jenv, jobject thiz)
@@ -484,6 +409,7 @@ static void InitCommandLine()
 			FCommandLine::Append(UTF8_TO_TCHAR(CommandLine));
 			FPlatformMisc::LowLevelOutputDebugStringf(TEXT("APK Commandline: %s"), FCommandLine::Get());
 		}
+
 #if !UE_BUILD_SHIPPING && USE_ANDROID_STANDALONE
 		UE_SET_LOG_VERBOSITY(LogAndroid, Log);
 		if (!GAndroidCommandLine.IsEmpty())
@@ -596,21 +522,31 @@ static void ApplyAndroidCompatConfigRules()
 }
 
 //Main function called from the android entry point
-#if USE_ANDROID_STANDALONE
+int32 AndroidMain(struct android_app* state);
+
 void* AndroidMain(void* param)
 {
 	struct android_app* state = (struct android_app*)param;
-
 	FTaskTagScope Scope(ETaskTag::EGameThread);
-	GGameThreadId = FPlatformTLS::GetCurrentThreadId();
-	GNativeAndroidApp = state;
-	check(GNativeAndroidApp);
 
-#else
+	if (GGameThreadId == 0)
+	{
+		GGameThreadId = FPlatformTLS::GetCurrentThreadId();
+		STANDALONE_DEBUG_LOGf(LogAndroid, TEXT("AndroidMain set current GGameThreadId=%d"), GGameThreadId);
+	}
+
+	if (EventThreadID == 0)
+	{
+		EventThreadID = GGameThreadId;
+		STANDALONE_DEBUG_LOGf(LogAndroid, TEXT("AndroidMain set current EventThreadID=%d"), EventThreadID);
+	}
+	AndroidMain(state);
+	return nullptr;
+}
+
+
 int32 AndroidMain(struct android_app* state)
 {
-#endif
-
 	BootTimingPoint("AndroidMain");
 
 	FPlatformMisc::LowLevelOutputDebugString(TEXT("Entered AndroidMain()\n"));
@@ -700,14 +636,7 @@ int32 AndroidMain(struct android_app* state)
 			FPlatformMisc::MemoryBarrier();
 		}
 	}
-	// make sure GEngineLoop::Exit() is always called.
-	struct EngineLoopCleanup
-	{
-		~EngineLoopCleanup()
-		{
-			EngineExit();
-		}
-	} Cleanup;
+
 
 	// read the command line file
 	InitCommandLine();
@@ -745,7 +674,7 @@ int32 AndroidMain(struct android_app* state)
 		// We don't wait for all of preinit to complete as PreLoadScreens will need to process events during preinit.
 
 		UE_LOG(LogAndroid, Log, TEXT("PreInit android HW window lock. bAppIsActive_EventThread=%d"), bAppIsActive_EventThread);
-		GAndroidWindowLock_Lock("AndroidMain lock");
+		GAndroidWindowLock.Lock();
 	}
 
 	FPlatformMisc::LowLevelOutputDebugString(TEXT("After GAndroidWindowLock in AndroidMain"));
@@ -760,11 +689,7 @@ int32 AndroidMain(struct android_app* state)
 	if (PreInitResult != 0)
 	{
 		checkf(false, TEXT("Engine Preinit Failed"));
-#if USE_ANDROID_STANDALONE
-		return nullptr;
-#else
 		return PreInitResult;
-#endif
 	}
 
 	// register callback for native window resize
@@ -821,7 +746,11 @@ int32 AndroidMain(struct android_app* state)
 		FAndroidStats::UpdateAndroidStats();
 
 		FAppEventManager::GetInstance()->Tick();
-		if (!FAppEventManager::GetInstance()->IsGamePaused())
+		if (!FAppEventManager::GetInstance()->IsGamePaused()
+#if USE_ANDROID_STANDALONE
+			//&& FAppEventManager::GetInstance()->IsGameInFocus()
+#endif
+			)
 		{
 			GEngineLoop.Tick();
 		}
@@ -848,18 +777,12 @@ int32 AndroidMain(struct android_app* state)
 
 	// exit out!
 	GEngineLoop.Exit();
-	
-#if !USE_ANDROID_STANDALONE
-	FPlatformMisc::RequestExit(true, TEXT("AndroidMain"));
-#endif
 
 	UE_LOG(LogAndroid, Log, TEXT("Exiting is over"));
 
-#if USE_ANDROID_STANDALONE
-	return nullptr;
-#else
-	return ErrorLevel;
-#endif
+	FPlatformMisc::RequestExit(true, TEXT("AndroidMain"));
+
+	return 0;
 }
 
 struct AChoreographer;
@@ -973,10 +896,11 @@ void FChoreographer::SetCallback(int64 Delay)
 	AChoreographer_postFrameCallbackDelayed_(choreographer, choreographer_callback, nullptr, Delay / 1000000);
 }
 
-static uint32 EventThreadID = 0;
-
 bool IsInAndroidEventThread()
 {
+	// Note: leave the commented out line for debug purposes.
+	//STANDALONE_DEBUG_LOGf(LogAndroid, TEXT("IsInAndroidEventThread(), GGameThreadId=%d, EventThreadID=%d, FPlatformTLS::GetCurrentThreadId()=%d"), GGameThreadId, EventThreadID, FPlatformTLS::GetCurrentThreadId());
+
 #if USE_ANDROID_STANDALONE
 	//@TODO: for now always return true to avoid check failures
 	return true;
@@ -989,6 +913,8 @@ bool IsInAndroidEventThread()
 static void* AndroidEventThreadWorker( void* param )
 {
 	pthread_setname_np(pthread_self(), "EventWorker");
+	STANDALONE_DEBUG_LOGf(LogAndroid, TEXT("AndroidEventThreadWorker(begin), GGameThreadId=%d, EventThreadID=%d, FPlatformTLS::GetCurrentThreadId()=%d"), GGameThreadId, EventThreadID, FPlatformTLS::GetCurrentThreadId());
+
 	EventThreadID = FPlatformTLS::GetCurrentThreadId();
 	FAndroidMisc::RegisterThreadName("EventWorker", EventThreadID);
 
@@ -1015,11 +941,9 @@ static void* AndroidEventThreadWorker( void* param )
 
 	TheChoreographer.SetupChoreographer();
 
-	GEventHandlerInitialized = true;
-
 	// window is initially invalid/locked.
 	UE_LOG(LogAndroid, Log, TEXT("AndroidEventThreadWorker, Initial HW window lock."));
-	GAndroidWindowLock_Lock("AndroidEventThreadWorker");
+	GAndroidWindowLock.Lock();
 
 	DEVELOPER_LOG_COMMANDCB_CASE(AndroidEventThreadWorker_BeforeWhile);
 
@@ -1035,7 +959,7 @@ static void* AndroidEventThreadWorker( void* param )
 	}
 	DEVELOPER_LOG_COMMANDCB_CASE(AndroidEventThreadWorker_AfterWhile);
 
-	GAndroidWindowLock_Unlock("AndroidEventThreadWorker");
+	GAndroidWindowLock.Unlock();
 
 	UE_LOG(LogAndroid, Log, TEXT("AndroidEventThreadWorker->Exiting"));
 	STANDALONE_DEBUG_LOGf(LogAndroid, TEXT("Exiting AndroidEventThreadWorker"));
@@ -1091,6 +1015,7 @@ void android_main(struct android_app* state)
 	//@todo android: replace with native activity, main loop off of UI thread, etc.
 	AndroidMain(state);
 }
+
 void* android_main(void* param)
 {
 	struct android_app* state = (struct android_app*)param;
@@ -1525,22 +1450,23 @@ FAppEventData::FAppEventData(ANativeWindow* WindowIn)
 	//check(WindowWidth >= 0 && WindowHeight >= 0);
 }
 
-
-
 // called when the app has window + resume.
 static void ActivateApp_EventThread()
 {
 	DEVELOPER_LOG_COMMANDCB_CASE(ActivateApp_EventThread);
 	if (bAppIsActive_EventThread)
 	{
-		// Seems this can occur.
 		STANDALONE_DEBUG_LOGf(LogAndroid, TEXT("UNEXPECTED -- event thread, activate app, ALREADY have HW window lock according to bAppIsActive_EventThread. bReadyToProcessEvents=%d, bAppIsActive_EventThread=%d"), bReadyToProcessEvents, bAppIsActive_EventThread);
+
+#if !USE_ANDROID_STANDALONE // SUSPECT TODO should remove the #if?
+		// Seems this can occur.
 		return;
+#endif
 	}
 
 	// Unlock window when we're ready.
 	UE_LOG(LogAndroid, Log, TEXT("event thread, activate app, unlocking HW window"));
-	GAndroidWindowLock_Unlock("ActivateApp_EventThread");
+	GAndroidWindowLock.Unlock();
 
 	// wake the GT up.
 	FAppEventManager::GetInstance()->EnqueueAppEvent(APP_EVENT_STATE_APP_ACTIVATED);
@@ -1570,7 +1496,7 @@ static void ActivateApp_EventThread()
 }
 
 extern void BlockRendering();
-// called whenever the app loses loses window or pause.
+// called whenever the app loses window or pause.
 static void SuspendApp_EventThread()
 {
 	DEVELOPER_LOG_COMMANDCB_CASE(SuspendApp_EventThread);
@@ -1584,7 +1510,7 @@ static void SuspendApp_EventThread()
 	// Lock the window, this prevents event thread from removing the window whilst the RHI initializes.
 
 	UE_LOG(LogAndroid, Log, TEXT("SuspendApp_EventThread -> event thread, suspending app, acquiring HW window lock."));
-	GAndroidWindowLock_Lock("SuspendApp_EventThread");
+	GAndroidWindowLock.Lock();
 	STANDALONE_DEBUG_LOGf(LogAndroid, TEXT("SuspendApp_EventThread -> event thread, suspending app, acquiring HW window lock. bReadyToProcessEvents=%d, bAppIsActive_EventThread=%d"), bReadyToProcessEvents, bAppIsActive_EventThread);
 
 	if (bReadyToProcessEvents == false)
@@ -1603,7 +1529,6 @@ static void SuspendApp_EventThread()
 		FPlatformProcess::ReturnSynchEventToPool(EventToDelete);
 	});
 	
-
 	// perform the delegates before the window handle is cleared.
 	// This ensures any tasks that require a window handle will have it before we block the RT on the invalid window.
 	FAppEventManager::GetInstance()->EnqueueAppEvent(APP_EVENT_RUN_CALLBACK, FAppEventData([EMDoneTrigger]()
@@ -1629,9 +1554,12 @@ static void SuspendApp_EventThread()
 	// wait for a period of time before blocking rendering
 	UE_LOG(LogAndroid, Log, TEXT("SuspendApp_EventThread -> , waiting for event manager to process. tid: %d"), FPlatformTLS::GetCurrentThreadId());
 #if USE_ANDROID_STANDALONE
-	EMDoneTrigger->Reset();
-#endif
+	//EMDoneTrigger->Reset();
+	bool bSuccess = EMDoneTrigger->Wait(240);
+#else
 	bool bSuccess = EMDoneTrigger->Wait(4000);
+#endif
+
 	float ElapsedTimeInMs_EMDoneTrigger_Wait = FPlatformTime::ToMilliseconds(FPlatformTime::Cycles() - StartCycles);
 	UE_CLOG(!bSuccess, LogAndroid, Log, TEXT("SuspendApp_EventThread -> ERROR: backgrounding callback, not responded in timely manner. EMDoneTrigger->Wait, waited '%f' ms"), (float)ElapsedTimeInMs_EMDoneTrigger_Wait);
 	UE_LOG(LogAndroid, Log, TEXT("SuspendApp_EventThread -> EMDoneTrigger->Wait, waited '%f' ms"), (float)ElapsedTimeInMs_EMDoneTrigger_Wait);
@@ -1693,12 +1621,13 @@ static void OnAppCommandCB(struct android_app* app, int32_t cmd)
 			}
 			checkf(win != NULL, TEXT("Engine APP_CMD_INIT_WINDOW with STANDALONE Failed since win is NULL!"));
 			FAppEventManager::GetInstance()->HandleWindowCreated_EventThread((ANativeWindow*)win);
+			bHasWindow = win != nullptr;
 		}
 #else
 		STANDALONE_DEBUG_LOGf(LogAndroid, TEXT("Case APP_CMD_INIT_WINDOW"));
 		FAppEventManager::GetInstance()->HandleWindowCreated_EventThread(app->pendingWindow);
-#endif
 		bHasWindow = true;
+#endif
 		UE_LOG(LogAndroid, Log, TEXT("Case APP_CMD_INIT_WINDOW: bHasWindow=%d, bHasFocus=%d, bIsResumed=%d"), bHasWindow, bHasFocus, bIsResumed);
 		if (bHasWindow && bHasFocus && bIsResumed)
 		{
@@ -1814,8 +1743,6 @@ static void OnAppCommandCB(struct android_app* app, int32_t cmd)
 		break;
 	case APP_CMD_RESUME:
 
-		DEVELOPER_LOG_COMMANDCB_CASE(APP_CMD_RESUME);
-
 		/**
 		 * Command from main thread: the app's activity has been resumed.
 		 */
@@ -2098,7 +2025,6 @@ bool WaitForAndroidLoseFocusEvent(double TimeoutSeconds)
 	return FAppEventManager::GetInstance()->WaitForEventInQueue(EAppEventState::APP_EVENT_STATE_WINDOW_LOST_FOCUS, TimeoutSeconds);
 }
 
-
 #if USE_ANDROID_STANDALONE
 pthread_t G_AndroidMainThread;
 
@@ -2175,7 +2101,7 @@ static void InitEvent()
 
 static void AppSuspended()
 {
-	STANDALONE_DEBUG_LOGf(LogAndroid, TEXT("LaunchAndroid received AppSuspended"));
+	STANDALONE_DEBUG_LOGf(LogAndroid, TEXT("LaunchAndroid received AppSuspended, GFrameCounter=%d"), (int32)GFrameCounter);
 	AndroidThunkCpp_Engine_SendEvent(EVENTTYPE_APP_SUSPENDED, TEXT(""), (int32)GFrameCounter, 0, 0.0f);
 }
 
@@ -2247,6 +2173,7 @@ JNI_METHOD void Java_com_epicgames_makeaar_GameActivityForMakeAAR_nativeMain(JNI
 	FCoreDelegates::OnInit.AddStatic(InitEvent);
 	FCoreDelegates::OnPostEngineInit.AddStatic(PostEngineInitEvent);
 	FCoreDelegates::OnFEngineLoopInitComplete.AddStatic(EngineLoopInitCompleteEvent);
+	FCoreDelegates::ApplicationWillEnterBackgroundDelegate.AddStatic(AppSuspended);
 	FCoreDelegates::OnBeginFrame.AddStatic(BeginFrameEvent);
 	FCoreDelegates::OnEndFrame.AddStatic(EndFrameEvent);
 	FCoreUObjectDelegates::PreLoadMap.AddStatic(PreLoadMapEvent);
@@ -2271,9 +2198,9 @@ JNI_METHOD jint Java_com_epicgames_makeaar_GameActivityForMakeAAR_nativeInputTou
 	TouchType type = TouchEnded;
 	switch (inType)
 	{
-	case 0: type = TouchBegan; break;
-	case 1: type = TouchMoved; break;
-	case 2: type = TouchEnded; break;
+		case 0: type = TouchBegan; break;
+		case 1: type = TouchMoved; break;
+		case 2: type = TouchEnded; break;
 	}
 
 	ANativeWindow* Window = (ANativeWindow*)FAndroidWindow::GetHardwareWindow_EventThread();
@@ -2379,7 +2306,7 @@ JNI_METHOD bool Java_com_epicgames_makeaar_Engine_nativeMountPak(JNIEnv* jenv, j
 
 JNI_METHOD void Java_com_epicgames_makeaar_GameActivityForMakeAAR_nativeAppCommand(JNIEnv* jenv, jobject thiz, jint cmd)
 {
-	STANDALONE_DEBUG_LOGf(LogAndroid, TEXT("MAKEAAR -> nativeAppCommand[%u, %s] -> EventThreadID=%d, GGameThreadId=%d, GResumeMainInit=%d, GEventHandlerInitialized=%d, FTaskGraphInterface::IsRunning()=%d, FPlatformTLS::GetCurrentThreadId()=%d"), 
+	STANDALONE_DEBUG_LOGf(LogAndroid, TEXT("GameActivityForMakeAAR -> nativeAppCommand[%u, %s] -> EventThreadID=%d, GGameThreadId=%d, GResumeMainInit=%d, GEventHandlerInitialized=%d, FTaskGraphInterface::IsRunning()=%d, FPlatformTLS::GetCurrentThreadId()=%d"), 
 		cmd,
 		get_APP_CMD_String(cmd),
 		EventThreadID, 
@@ -2407,4 +2334,4 @@ extern "C" {
     }
 };
 
-#endif
+#endif //  USE_ANDROID_LAUNCH
diff --git a/Engine/Source/Runtime/OpenGLDrv/Private/Android/AndroidEGL.cpp b/Engine/Source/Runtime/OpenGLDrv/Private/Android/AndroidEGL.cpp
--- a/Engine/Source/Runtime/OpenGLDrv/Private/Android/AndroidEGL.cpp
+++ b/Engine/Source/Runtime/OpenGLDrv/Private/Android/AndroidEGL.cpp
@@ -1260,9 +1260,7 @@ bool AndroidEGL::IsOfflineSurfaceRequired()
 }
 
 ///
-//extern FCriticalSection GAndroidWindowLock;
-extern void GAndroidWindowLock_Lock(FString calledBy);
-extern void GAndroidWindowLock_Unlock(FString calledBy);
+extern FCriticalSection GAndroidWindowLock;
 
 void BlockOnLostWindowRenderCommand(TSharedPtr<FEvent, ESPMode::ThreadSafe> RTBlockedTrigger)
 {
@@ -1292,7 +1290,7 @@ void BlockOnLostWindowRenderCommand(TSharedPtr<FEvent, ESPMode::ThreadSafe> RTBl
 
 		RTBlockedTrigger->Trigger();
 
-		GAndroidWindowLock_Lock("BlockOnLostWindowRenderCommand Vulkan");
+		GAndroidWindowLock.Lock();
 		UE_LOG(LogAndroid, Log, TEXT("RendererBlock acquired window lock"));
 		const auto& OnReinitWindowCallback = FAndroidMisc::GetOnReInitWindowCallback();
 		if (OnReinitWindowCallback)
@@ -1300,17 +1298,17 @@ void BlockOnLostWindowRenderCommand(TSharedPtr<FEvent, ESPMode::ThreadSafe> RTBl
 			OnReinitWindowCallback(FAndroidWindow::GetHardwareWindow_EventThread());
 			UE_LOG(LogAndroid, Log, TEXT("RendererBlock updating window"));
 		}
-		GAndroidWindowLock_Unlock("BlockOnLostWindowRenderCommand Vulkan");
+		GAndroidWindowLock.Unlock();
 	}
 	else
 	{
 		RunOnGLRenderContextThread([&] {
 			RTBlockedTrigger->Trigger();
-			GAndroidWindowLock_Lock("BlockOnLostWindowRenderCommand");
+			GAndroidWindowLock.Lock();
 			UE_LOG(LogAndroid, Log, TEXT("RendererBlock acquired window lock"));
 			AndroidEGL::GetInstance()->SetRenderContextWindowSurface();
 			UE_LOG(LogAndroid, Log, TEXT("RendererBlock updating window"));
-			GAndroidWindowLock_Unlock("BlockOnLostWindowRenderCommand");
+			GAndroidWindowLock.Unlock();
 		}, true);
 	}
 	UE_LOG(LogAndroid, Log, TEXT("RendererBlock released window lock"));
@@ -1328,9 +1326,7 @@ void SetSharedContextGameCommand(TSharedPtr<FEvent, ESPMode::ThreadSafe> GTBlock
 extern bool IsInAndroidEventThread();
 void BlockRendering()
 {
-#if !USE_ANDROID_STANDALONE
 	check(IsInAndroidEventThread());
-#endif
 	check(GIsRHIInitialized);
 
 	UE_LOG(LogAndroid, Log, TEXT("Blocking renderer on suspended window."));
