OverrideProjectName is no longer needed with ASIS
#android

diff --git a/Engine/Source/Runtime/Launch/Private/Android/LaunchAndroid.cpp b/Engine/Source/Runtime/Launch/Private/Android/LaunchAndroid.cpp
--- a/Engine/Source/Runtime/Launch/Private/Android/LaunchAndroid.cpp
+++ b/Engine/Source/Runtime/Launch/Private/Android/LaunchAndroid.cpp
@@ -2157,17 +2157,9 @@ JNI_METHOD void Java_com_epicgames_makeaar_GameActivityForMakeAAR_nativeMain(JNI
 	BootTimingPoint("native_main");
 	GIsGameAgnosticExe = true;
 
-	if (projectModule)
-	{
-		FString name = FJavaHelper::FStringFromParam(jenv, projectModule);
-		STANDALONE_DEBUG_LOGf(LogAndroid, TEXT("Java_com_epicgames_makeaar_GameActivityForMakeAAR_nativeMain : OverrideProjectModule to 'makeaar_GameActivityForMakeAAR', projectModule name is '%s'"), *name);
-		GEngineLoop.OverrideProjectModule(FJavaHelper::FStringFromParam(jenv, projectModule), TEXT("makeaar_GameActivityForMakeAAR"));
-	}
-	else
-	{
-		STANDALONE_DEBUG_LOGf(LogAndroid, TEXT("Java_com_epicgames_makeaar_GameActivityForMakeAAR_nativeMain : use current ProjectModule"));
+	FString projectModuleName = FJavaHelper::FStringFromParam(jenv, projectModule);
 
-	}
+	STANDALONE_DEBUG_LOGf(LogAndroid, TEXT("Java_com_epicgames_makeaar_GameActivityForMakeAAR_nativeMain : use current!, requesting ProjectModule: %s"), *projectModuleName);
 
 	// register some delegates we want to pass back
 	FCoreDelegates::OnInit.AddStatic(InitEvent);
