Fixing disregard for GC set never being properly disabled
+ Removing gc.UseDisregardForGCOnDedicatedServers setting

diff --git a/Engine/Config/BaseEngine.ini b/Engine/Config/BaseEngine.ini
--- a/Engine/Config/BaseEngine.ini
+++ b/Engine/Config/BaseEngine.ini
@@ -1539,7 +1539,6 @@ gc.CreateGCClusters=True
 gc.MinGCClusterSize=5
 gc.AssetClustreringEnabled=False
 gc.ActorClusteringEnabled=False
-gc.UseDisregardForGCOnDedicatedServers=False
 gc.MultithreadedDestructionEnabled=True
 gc.VerifyUObjectsAreNotFGCObjects=False
 gc.GarbageEliminationEnabled=True
diff --git a/Engine/Source/Runtime/CoreUObject/Private/UObject/UObjectArray.cpp b/Engine/Source/Runtime/CoreUObject/Private/UObject/UObjectArray.cpp
--- a/Engine/Source/Runtime/CoreUObject/Private/UObject/UObjectArray.cpp
+++ b/Engine/Source/Runtime/CoreUObject/Private/UObject/UObjectArray.cpp
@@ -178,12 +178,17 @@ void FUObjectArray::CloseDisregardForGC()
 
 	OpenForDisregardForGC = false;
 	GIsInitialLoad = false;
+
+	checkf(!DisregardForGCEnabled() || !GIsEditor, TEXT("Disregard For GC Set can't be enabled when running the editor"));
+	checkf(DisregardForGCEnabled() || (ObjFirstGCIndex == 0 && ObjLastNonGCIndex == -1), TEXT("Disregard for GC Set is not properly disabled (FirstGCIndex = %d, LastNonGCIndex = %d"), ObjFirstGCIndex, ObjLastNonGCIndex);
 }
 
 void FUObjectArray::DisableDisregardForGC()
 {
 	if (!GExitPurge && (ObjFirstGCIndex > 0 || DisregardForGCEnabled()))
 	{
+		checkf(!IsAsyncLoading(), TEXT("Disregard for GC Set can't be safely disabled while async loading. Consider calling FlushAsyncLoading() first or using gc.MaxObjectsNotConsideredByGC=0 ini setting instead."));
+
 		void OnDisregardForGCSetDisabled(int32 NumObjects);
 		// If disregard for GC was already closed then ObjFirstGCIndex is the number of objects we need to scan, otherwise disregard for GC is still open and we need to scan all objects
 		int32 NumDisregardForGCObjects = ObjFirstGCIndex > 0 ? ObjFirstGCIndex : GetObjectArrayNum();
@@ -195,6 +200,7 @@ void FUObjectArray::DisableDisregardForGC()
 
 	MaxObjectsNotConsideredByGC = 0;
 	ObjFirstGCIndex = 0;
+	ObjLastNonGCIndex = -1;
 	if (IsOpenForDisregardForGC())
 	{
 		CloseDisregardForGC();
diff --git a/Engine/Source/Runtime/Engine/Classes/Engine/CoreSettings.h b/Engine/Source/Runtime/Engine/Classes/Engine/CoreSettings.h
--- a/Engine/Source/Runtime/Engine/Classes/Engine/CoreSettings.h
+++ b/Engine/Source/Runtime/Engine/Classes/Engine/CoreSettings.h
@@ -214,11 +214,6 @@ protected:
 		ToolTip = "Whether to allow levels to create actor clusters for GC."))
 	uint32 ActorClusteringEnabled : 1;
 
-	UPROPERTY(EditAnywhere, config, Category = Optimization, meta = (
-		ConsoleVariable = "gc.UseDisregardForGCOnDedicatedServers", DisplayName = "Use DisregardForGC On Dedicated Servers",
-		ToolTip = "If false, DisregardForGC will be disabled for dedicated servers."))
-	uint32 UseDisregardForGCOnDedicatedServers : 1;
-
 	UPROPERTY(EditAnywhere, config, Category = Debug, meta = (
 		ConsoleVariable = "gc.VerifyUObjectsAreNotFGCObjects", DisplayName = "Verify UObjects Are Not FGCObjects",
 		ToolTip = "If true, the engine will throw a warning when it detects a UObject-derived class which also derives from FGCObject or any of its members is derived from FGCObject."))
diff --git a/Engine/Source/Runtime/Engine/Private/CoreSettings.cpp b/Engine/Source/Runtime/Engine/Private/CoreSettings.cpp
--- a/Engine/Source/Runtime/Engine/Private/CoreSettings.cpp
+++ b/Engine/Source/Runtime/Engine/Private/CoreSettings.cpp
@@ -211,7 +211,6 @@ UGarbageCollectionSettings::UGarbageCollectionSettings()
 	MinGCClusterSize = 5;
 	AssetClusteringEnabled = true;
 	ActorClusteringEnabled = true;
-	UseDisregardForGCOnDedicatedServers = false;
 	VerifyUObjectsAreNotFGCObjects = true;
 	GarbageEliminationEnabled = false;
 	DumpObjectCountsToLogWhenMaxObjectLimitExceeded = false;
diff --git a/Engine/Source/Runtime/Launch/Private/LaunchEngineLoop.cpp b/Engine/Source/Runtime/Launch/Private/LaunchEngineLoop.cpp
--- a/Engine/Source/Runtime/Launch/Private/LaunchEngineLoop.cpp
+++ b/Engine/Source/Runtime/Launch/Private/LaunchEngineLoop.cpp
@@ -300,14 +300,6 @@ UE_TRACE_EVENT_END()
 
 bool GIsConsoleExecutable = false;
 
-int32 GUseDisregardForGCOnDedicatedServers = 1;
-static FAutoConsoleVariableRef CVarUseDisregardForGCOnDedicatedServers(
-	TEXT("gc.UseDisregardForGCOnDedicatedServers"),
-	GUseDisregardForGCOnDedicatedServers,
-	TEXT("If false, DisregardForGC will be disabled for dedicated servers."),
-	ECVF_Default
-);
-
 static TAutoConsoleVariable<int32> CVarDoAsyncEndOfFrameTasksRandomize(
 	TEXT("tick.DoAsyncEndOfFrameTasks.Randomize"),
 	0,
@@ -3412,12 +3404,6 @@ int32 FEngineLoop::PreInitPreStartupScreen(const TCHAR* CmdLine)
 	PreInitContext.bWithConfigPatching = bWithConfigPatching;
 	PreInitContext.bHasEditorToken = bHasEditorToken;
 #if WITH_ENGINE
-	bool bDisableDisregardForGC = bHasEditorToken;
-	if (bIsRunningAsDedicatedServer)
-	{
-		bDisableDisregardForGC |= FPlatformProperties::RequiresCookedData() && (GUseDisregardForGCOnDedicatedServers == 0);
-	}
-	PreInitContext.bDisableDisregardForGC = bDisableDisregardForGC;
 	PreInitContext.bIsRegularClient = bIsRegularClient;
 #endif // WITH_ENGINE
 	PreInitContext.bIsPossiblyUnrecognizedCommandlet = bIsRegularClient && TokenToForward.Len() && !TokenToForward.Contains(TEXT("-"));
@@ -3483,7 +3469,6 @@ int32 FEngineLoop::PreInitPostStartupScreen(const TCHAR* CmdLine)
 	bool bWithConfigPatching = PreInitContext.bWithConfigPatching;
 	bool bHasEditorToken = PreInitContext.bHasEditorToken;
 #if WITH_ENGINE
-	bool bDisableDisregardForGC = PreInitContext.bDisableDisregardForGC;
 	bool bIsRegularClient = PreInitContext.bIsRegularClient;
 #endif // WITH_ENGINE
 	bool bIsPossiblyUnrecognizedCommandlet = PreInitContext.bIsPossiblyUnrecognizedCommandlet;
@@ -3751,13 +3736,6 @@ int32 FEngineLoop::PreInitPostStartupScreen(const TCHAR* CmdLine)
 
 	FEmbeddedCommunication::ForceTick(7);
 
-	// Setup GC optimizations
-	if (bDisableDisregardForGC)
-	{
-		SCOPED_BOOT_TIMING("DisableDisregardForGC");
-		GUObjectArray.DisableDisregardForGC();
-	}
-
 	SlowTask.EnterProgressFrame(10);
 
 	{
diff --git a/Engine/Source/Runtime/Launch/Public/LaunchEngineLoop.h b/Engine/Source/Runtime/Launch/Public/LaunchEngineLoop.h
--- a/Engine/Source/Runtime/Launch/Public/LaunchEngineLoop.h
+++ b/Engine/Source/Runtime/Launch/Public/LaunchEngineLoop.h
@@ -21,6 +21,7 @@ struct FPreInitContext
 	bool bDumpEarlyPakFileReads = false;
 	bool bForceQuitAfterEarlyReads = false;
 	bool bWithConfigPatching = false;
+	UE_DEPRECATED(5.5, "Use gc.MaxObjectsNotConsideredByGC=0 ini setting instead")
 	bool bDisableDisregardForGC = false;
 	bool bHasEditorToken = false;
 	bool bIsRegularClient = false;
