RunCommandlet console command. Allows for faster iteration when debugging commandlets in the editor (e.g via hot reload)

diff --git a/Engine/Source/Runtime/Launch/Private/LaunchEngineLoop.cpp b/Engine/Source/Runtime/Launch/Private/LaunchEngineLoop.cpp
--- a/Engine/Source/Runtime/Launch/Private/LaunchEngineLoop.cpp
+++ b/Engine/Source/Runtime/Launch/Private/LaunchEngineLoop.cpp
@@ -119,6 +119,7 @@
 	#include "Editor/EditorEngine.h"
 	#include "UnrealEdMisc.h"
 	#include "UnrealEdGlobals.h"
+	#include "UObject/StrongObjectPtr.h"
 	#include "Editor/UnrealEdEngine.h"
 	#include "Settings/EditorExperimentalSettings.h"
 	#include "PIEPreviewDeviceProfileSelectorModule.h"
@@ -748,6 +749,64 @@ public:
 	}
 };
 
+#if WITH_ENGINE && WITH_EDITOR
+static FAutoConsoleCommand CmdRunCommandlet(
+	TEXT("RunCommandlet"),
+	TEXT("<Commandlet> <Args...>. NOTE: This is for debugging/iteration purposes only! Running commandlets multiple times in an editor session may result in issues!"),
+	FConsoleCommandWithArgsDelegate::CreateLambda( [](const TArray<FString>& Args) 
+		{
+			if ( Args.Num() == 0 )
+			{
+				UE_LOG(LogInit, Error, TEXT("Need to specify a commandlet name to run!"));
+				return;
+			}
+
+			FString CommandletName = Args[0];
+			UClass* CommandletClass = nullptr;
+
+			CommandletClass = Cast<UClass>(StaticFindFirstObject(UClass::StaticClass(), *CommandletName, EFindFirstObjectOptions::None, ELogVerbosity::Warning, TEXT("looking for commandlet")));
+			int32 PeriodIdx;
+			if (!CommandletClass && CommandletName.FindChar('.', PeriodIdx))
+			{
+				// try to load module for commandlet specified before a period.
+				FModuleManager::Get().LoadModule(*CommandletName.Left(PeriodIdx));
+				CommandletClass = FindFirstObject<UClass>(*CommandletName, EFindFirstObjectOptions::None, ELogVerbosity::Warning, TEXT("Looking for commandlet class"));
+			}
+			if (!CommandletClass)
+			{
+				UE_LOG(LogInit, Error, TEXT("Commandlet class '%s' not found!"), *CommandletName);
+				return;
+			}
+			TStrongObjectPtr<UCommandlet> Commandlet( NewObject<UCommandlet>(GetTransientPackage(), CommandletClass) );
+
+			// Execute the commandlet.
+			double CommandletExecutionStartTime = FPlatformTime::Seconds();
+			
+			// Commandlets don't always handle -run= properly in the commandline so we'll provide them
+			// with a custom version that doesn't have it.
+			TArray<FString> CmdlineArgs = Args;
+			CmdlineArgs.RemoveAt(0);
+			FString CommandletCommandLine = FString::Join(CmdlineArgs, TEXT(" "));
+			Commandlet->ParseParms(*CommandletCommandLine);
+
+			PRIVATE_GRunningCommandletClass = CommandletClass;
+
+			FCoreDelegates::OnCommandletPreMain.Broadcast();
+			int32 ErrorLevel;
+			{
+				TRACE_CPUPROFILER_EVENT_SCOPE_TEXT(*WriteToString<512>(TEXT("Commandlet Main "), Commandlet->GetFName()));
+				FTrackedActivityScope CommandletActivity(FTrackedActivity::GetEngineActivity(), *FString::Printf(TEXT("Running %s"), *Commandlet->GetName()), false, FTrackedActivity::ELight::Green);
+				ErrorLevel = Commandlet->Main(CommandletCommandLine);
+			}
+			FCoreDelegates::OnCommandletPostMain.Broadcast();
+
+			PRIVATE_GRunningCommandletClass = nullptr;
+
+			double CommandletExecutionTime = FPlatformTime::Seconds() - CommandletExecutionStartTime;
+			UE_LOG(LogInit, Display, LINE_TERMINATOR TEXT("Commandlet returned %d. Execution time:  %.2f seconds"), ErrorLevel, CommandletExecutionTime);
+		})
+	);
+#endif
 
 #if WITH_APPLICATION_CORE
 static TUniquePtr<FOutputDeviceConsole>	GScopedLogConsole;
