Move pipelinestatecache shutdown back to RHIExit
Add a wait for PSO compile tasks to complete before AppPreExit.
 - Clear pending PSO precache tasks before AppPreExit to avoid potential waits.

diff --git a/Engine/Source/Runtime/Launch/Private/LaunchEngineLoop.cpp b/Engine/Source/Runtime/Launch/Private/LaunchEngineLoop.cpp
--- a/Engine/Source/Runtime/Launch/Private/LaunchEngineLoop.cpp
+++ b/Engine/Source/Runtime/Launch/Private/LaunchEngineLoop.cpp
@@ -193,6 +193,7 @@
 	#include "DynamicResolutionState.h"
 	#include "EngineModule.h"
 	#include "DumpGPU.h"
+	#include "PSOPrecacheMaterial.h"
 	#include "VisualizeTexture.h"
 
 #if !UE_SERVER
@@ -5055,12 +5056,12 @@ void FEngineLoop::Exit()
 #endif // WITH_EDITOR
 	FModuleManager::Get().UnloadModule("WorldBrowser", true);
 
-	// Clean up all cached pipelines, before the garbage collector shuts down.
-	ENQUEUE_RENDER_COMMAND(ShutdownPipelineStateCache)([](FRHICommandListImmediate& RHICmdList)
-	{
-		PipelineStateCache::Shutdown();
-	});
-	FlushRenderingCommands();
+
+#if WITH_ENGINE	
+	// Reset any in progress PSO compile requests, reduces pipelinestatecache task wait time.
+	ClearMaterialPSORequests();
+#endif
+	PipelineStateCache::WaitForAllTasks();
 
 	AppPreExit();
 
diff --git a/Engine/Source/Runtime/RHI/Private/DynamicRHI.cpp b/Engine/Source/Runtime/RHI/Private/DynamicRHI.cpp
--- a/Engine/Source/Runtime/RHI/Private/DynamicRHI.cpp
+++ b/Engine/Source/Runtime/RHI/Private/DynamicRHI.cpp
@@ -420,6 +420,9 @@ void RHIExit()
 {
 	if (!GUsingNullRHI && GDynamicRHI != NULL)
 	{
+		// Clean up all cached pipelines
+		PipelineStateCache::Shutdown();
+
 		// Flush any potential commands queued before we shut things down.
 		FRHICommandListExecutor::GetImmediateCommandList().ImmediateFlush(EImmediateFlushType::FlushRHIThread);
 
diff --git a/Engine/Source/Runtime/RHI/Private/PipelineStateCache.cpp b/Engine/Source/Runtime/RHI/Private/PipelineStateCache.cpp
--- a/Engine/Source/Runtime/RHI/Private/PipelineStateCache.cpp
+++ b/Engine/Source/Runtime/RHI/Private/PipelineStateCache.cpp
@@ -3349,13 +3349,19 @@ void DumpPipelineCacheStats()
 static TMap<uint32, FRHIVertexDeclaration*> GVertexDeclarationCache;
 static FCriticalSection GVertexDeclarationLock;
 
-void PipelineStateCache::Shutdown()
+void PipelineStateCache::WaitForAllTasks()
 {
 	GComputePipelineCache.WaitTasksComplete();
 	GWorkGraphPipelineCache.WaitTasksComplete();
 	GGraphicsPipelineCache.WaitTasksComplete();
 	GPrecacheGraphicsPipelineCache.WaitTasksComplete();
 	GPrecacheComputePipelineCache.WaitTasksComplete();
+}
+
+void PipelineStateCache::Shutdown()
+{
+	WaitForAllTasks();
+
 #if RHI_RAYTRACING
 	GRayTracingPipelineCache.Shutdown();
 #endif
diff --git a/Engine/Source/Runtime/RHI/Public/PipelineStateCache.h b/Engine/Source/Runtime/RHI/Public/PipelineStateCache.h
--- a/Engine/Source/Runtime/RHI/Public/PipelineStateCache.h
+++ b/Engine/Source/Runtime/RHI/Public/PipelineStateCache.h
@@ -134,6 +134,9 @@ namespace PipelineStateCache
 
 	extern RHI_API void ReportFrameHitchToCSV();
 
+	// Waits for any pending tasks to complete.
+	extern RHI_API void WaitForAllTasks();
+
 	/* Clears all pipeline cached state. Called on shutdown, calling GetAndOrCreate after this will recreate state */
 	extern RHI_API void Shutdown();
 
