Fix Client Crash in PipelineStateCache::GetAndOrCreateComputePipelineState

Move pipelinestatecache shutdown back to RHIExit
Add a wait for PSO compile tasks to complete before AppPreExit.
 - Clear pending PSO precache tasks before AppPreExit to avoid potential waits.

diff --git a/Engine/Source/Runtime/Launch/Private/LaunchEngineLoop.cpp b/Engine/Source/Runtime/Launch/Private/LaunchEngineLoop.cpp
--- a/Engine/Source/Runtime/Launch/Private/LaunchEngineLoop.cpp
+++ b/Engine/Source/Runtime/Launch/Private/LaunchEngineLoop.cpp
@@ -4935,6 +4935,14 @@ int32 FEngineLoop::Init()
 }
 
 
+// 5.4.2 local change to avoid modifying public headers
+namespace PipelineStateCache
+{
+	// Waits for any pending tasks to complete.
+	extern RHI_API void WaitForAllTasks();
+
+}
+
 void FEngineLoop::Exit()
 {
 	STAT_ADD_CUSTOMMESSAGE_NAME( STAT_NamedMarker, TEXT( "EngineLoop.Exit" ) );
@@ -5062,6 +5070,8 @@ void FEngineLoop::Exit()
 	// Reset any in progress PSO compile requests, reduces pipelinestatecache task wait time.
 	ClearMaterialPSORequests();
 #endif
+
+	// Wait for any pending tasks to complete.
 	PipelineStateCache::WaitForAllTasks();
 
 	AppPreExit();
diff --git a/Engine/Source/Runtime/RHI/Private/PipelineStateCache.cpp b/Engine/Source/Runtime/RHI/Private/PipelineStateCache.cpp
--- a/Engine/Source/Runtime/RHI/Private/PipelineStateCache.cpp
+++ b/Engine/Source/Runtime/RHI/Private/PipelineStateCache.cpp
@@ -20,6 +20,14 @@ PipelineStateCache.cpp: Pipeline state cache implementation.
 #include "Stats/StatsTrace.h"
 #include "Templates/TypeHash.h"
 
+// 5.4.2 local change to avoid modifying public headers
+namespace PipelineStateCache
+{
+	// Waits for any pending tasks to complete.
+	extern RHI_API void WaitForAllTasks();
+
+}
+
 // perform cache eviction each frame, used to stress the system and flush out bugs
 #define PSO_DO_CACHE_EVICT_EACH_FRAME 0
 
