Added a warning log in the Address Sanitizer callbacks to indicate if GIsCriticalError is set. Many subsystems skip appropriate clean-up routines (namely ~UObjectBase) when a critical error is flagged (i.e. from a GPU hang) and it's not unlikely that we run into access of freed/invalid memory after that.

This is intended to serve as a warning that the ASan error may be spurious and is the result of a previous critical/fatal error that has not forced the application to exit yet.

#rnx

diff --git a/Engine/Source/Runtime/Launch/Private/Windows/LaunchWindows.cpp b/Engine/Source/Runtime/Launch/Private/Windows/LaunchWindows.cpp
--- a/Engine/Source/Runtime/Launch/Private/Windows/LaunchWindows.cpp
+++ b/Engine/Source/Runtime/Launch/Private/Windows/LaunchWindows.cpp
@@ -22,6 +22,10 @@ DEFINE_LOG_CATEGORY_STATIC(LogASan, Log, All);
 
 static void ASanErrorCallback(const char* ErrorStr)
 {
+	if (GIsCriticalError)
+	{
+		UE_LOG(LogASan, Warning, TEXT("Critical error raised prior to ASan error; engine memory may be unstable in this state!"));
+	}
 	UE_LOG(LogASan, Fatal, TEXT("ASan Error: %s"), ANSI_TO_TCHAR(ErrorStr));
 }
 
