de-duplicate the line after an incorrect merge

#rnx

diff --git a/Engine/Source/Runtime/Launch/Private/Windows/LaunchWindows.cpp b/Engine/Source/Runtime/Launch/Private/Windows/LaunchWindows.cpp
--- a/Engine/Source/Runtime/Launch/Private/Windows/LaunchWindows.cpp
+++ b/Engine/Source/Runtime/Launch/Private/Windows/LaunchWindows.cpp
@@ -26,7 +26,6 @@ static void ASanErrorCallback(const char* ErrorStr)
 	{
 		UE_LOG(LogASan, Warning, TEXT("Critical error raised prior to ASan error; engine memory may be unstable in this state!"));
 	}
-	UE_LOG(LogASan, Fatal, TEXT("ASan Error: %s"), ANSI_TO_TCHAR(ErrorStr));
 	/* Marking the end of sanitizer report for Gauntlet. */
 	UE_LOG(LogASan, Fatal, TEXT("ASan Error: %s\nEnd of Address Sanitizer report"), ANSI_TO_TCHAR(ErrorStr));
 }
