Clean up condition in a UE_CLOG

diff --git a/Engine/Source/Runtime/Launch/Private/LaunchEngineLoop.cpp b/Engine/Source/Runtime/Launch/Private/LaunchEngineLoop.cpp
--- a/Engine/Source/Runtime/Launch/Private/LaunchEngineLoop.cpp
+++ b/Engine/Source/Runtime/Launch/Private/LaunchEngineLoop.cpp
@@ -1654,7 +1654,7 @@ bool IsServerDelegateForOSS(FName WorldContextHandle)
 #if !WITH_DEV_AUTOMATION_TESTS
 			// Not having a world to make the right determination is a bad thing
 			// In the editor during PIE this will confuse the individual PIE windows and their associated online components
-			UE_CLOG((World == nullptr), LogInit, Error, TEXT("Failed to determine if OSS is server in PIE, OSS requests will fail"));
+			UE_CLOG(World == nullptr, LogInit, Error, TEXT("Failed to determine if OSS is server in PIE, OSS requests will fail"));
 #endif
 		}
 	}
