[Launch]
* Hacky fix for merged modules to be able to boot... Forcing load of Common module for binary loaded. This is probably not the right place and maybe the common dll should be loaded automatically through dll dependencies

diff --git a/Engine/Source/Runtime/Launch/Launch.Build.cs b/Engine/Source/Runtime/Launch/Launch.Build.cs
--- a/Engine/Source/Runtime/Launch/Launch.Build.cs
+++ b/Engine/Source/Runtime/Launch/Launch.Build.cs
@@ -260,5 +260,10 @@ public class Launch : ModuleRules
 			PrivateDefinitions.Add(string.Format("COMPILED_IN_COMPATIBLE_CL={0}", Target.Version.EffectiveCompatibleChangelist));
 			PrivateDefinitions.Add(string.Format("COMPILED_IN_BRANCH_NAME={0}", (Target.Version.BranchName == null || Target.Version.BranchName.Length == 0) ? "UE" : Target.Version.BranchName));
 		}
+
+		if (Target.bMergeModules)
+		{
+			PrivateDefinitions.Add("UE_MERGED_MODULES=1");
+		}
 	}
 }
diff --git a/Engine/Source/Runtime/Launch/Private/LaunchEngineLoop.cpp b/Engine/Source/Runtime/Launch/Private/LaunchEngineLoop.cpp
--- a/Engine/Source/Runtime/Launch/Private/LaunchEngineLoop.cpp
+++ b/Engine/Source/Runtime/Launch/Private/LaunchEngineLoop.cpp
@@ -295,6 +295,10 @@ class FFeedbackContext;
 #include "IOS/IOSAppDelegate.h"
 #endif
 
+#ifndef UE_MERGED_MODULES
+#define UE_MERGED_MODULES 0
+#endif
+
 #if CPUPROFILERTRACE_ENABLED
 UE_TRACE_EVENT_BEGIN(Cpu, Frame, NoSync)
 	UE_TRACE_EVENT_FIELD(UE::Trace::WideString, Name)
@@ -2161,17 +2165,29 @@ int32 FEngineLoop::PreInitPreStartupScreen(const TCHAR* CmdLine)
 		ProjectBinariesRootDirectory = FPlatformMisc::ProjectDir();
 #if !IS_MONOLITHIC
 		FPlatformMisc::GetEngineAndProjectAbsoluteDirsFromExecutable(ProjectBinariesRootDirectory, EngineBinariesRootDirectory);
-		// Loading preinit module if it exists. PreInit module can contain things like decryption logic for pak files and things that is project specific but needs to initialize very early
-		TArray<FString> FileNames = {
+		
+		// Loading preinit/common module if it exists.
+		// PreInit module can contain things like decryption logic for pak files and things that is project specific but needs to initialize very early
+		// Common module is a merged binary containing all the important modules for the project
+		TArray<FString> FileNames =
+		{
+#if UE_MERGED_MODULES
+			FString::Printf(TEXT("%s-Common.%s"), FPlatformProcess::ExecutableName(), FPlatformProcess::GetModuleExtension()),
+			FString::Printf(TEXT("%s-Common-%s-%s.%s"), FPlatformProcess::ExecutableName(), FPlatformProcess::GetBinariesSubdirectory(), LexToString(FApp::GetBuildConfiguration()), FPlatformProcess::GetModuleExtension()),
+#else
 			FString::Printf(TEXT("%s-%sPreInit.%s"), FPlatformProcess::ExecutableName(), FApp::GetProjectName(), FPlatformProcess::GetModuleExtension()),
 			FString::Printf(TEXT("%s-%sPreInit-%s-%s.%s"), *FApp::GetName(), FApp::GetProjectName(), FPlatformProcess::GetBinariesSubdirectory(), LexToString(FApp::GetBuildConfiguration()), FPlatformProcess::GetModuleExtension())
+#endif
 		};
+
+		FTaskTagScope scope(ETaskTag::EStaticInit);
+
 		for (const FString& FileName : FileNames)
 		{
-			FString PreInitModule = FPaths::Combine(ProjectBinariesRootDirectory, "Binaries", FPlatformProcess::GetBinariesSubdirectory(), FileName);
-			if (FPaths::FileExists(PreInitModule))
+			FString ModulePath = FPaths::Combine(ProjectBinariesRootDirectory, "Binaries", FPlatformProcess::GetBinariesSubdirectory(), FileName);
+			if (FPaths::FileExists(ModulePath))
 			{
-				FPlatformProcess::GetDllHandle(*PreInitModule);
+				FPlatformProcess::GetDllHandle(*ModulePath);
 				break;
 			}
 		}
