Add SharedPreferences accessors for Android

#android

diff --git a/Engine/Build/Android/Java/src/com/epicgames/unreal/GameActivity.java.template b/Engine/Build/Android/Java/src/com/epicgames/unreal/GameActivity.java.template
--- a/Engine/Build/Android/Java/src/com/epicgames/unreal/GameActivity.java.template
+++ b/Engine/Build/Android/Java/src/com/epicgames/unreal/GameActivity.java.template
@@ -4760,20 +4760,59 @@ public class GameActivity extends $${gameActivitySuperClass}$$ implements Surfac
 		});
 	}
 
-	public void AndroidThunkJava_SetCellularPreference(int value)
+	public boolean AndroidThunkJava_HasSharedPreference(String group, String key)
 	{
 		Context context = getApplicationContext();
-		SharedPreferences preferences = context.getSharedPreferences("CellularNetworkPreferences", context.MODE_PRIVATE);
+		SharedPreferences preferences = context.getSharedPreferences(group, context.MODE_PRIVATE);
+		return preferences.contains(key);
+	}
+
+	public void AndroidThunkJava_SetSharedPreferenceBoolean(String group, String key, boolean value)
+	{
+		Context context = getApplicationContext();
+		SharedPreferences preferences = context.getSharedPreferences(group, context.MODE_PRIVATE);
+		SharedPreferences.Editor editor = preferences.edit();
+		editor.putBoolean(key, value);
+		editor.commit();
+	}
+
+	public boolean AndroidThunkJava_GetSharedPreferenceBoolean(String group, String key, boolean defaultValue)
+	{
+		Context context = getApplicationContext();
+		SharedPreferences preferences = context.getSharedPreferences(group, context.MODE_PRIVATE);
+		return preferences.getBoolean(key, defaultValue);
+	}
+
+	public void AndroidThunkJava_SetSharedPreferenceInt(String group, String key, int value)
+	{
+		Context context = getApplicationContext();
+		SharedPreferences preferences = context.getSharedPreferences(group, context.MODE_PRIVATE);
+		SharedPreferences.Editor editor = preferences.edit();
+		editor.putInt(key, key, value);
+		editor.commit();
+	}
+
+	public int AndroidThunkJava_GetSharedPreferenceInt(String group, String key, int defaultValue)
+	{
+		Context context = getApplicationContext();
+		SharedPreferences preferences = context.getSharedPreferences(group, context.MODE_PRIVATE);
+		return preferences.getInt(key, defaultValue);
+	}
+
+	public void AndroidThunkJava_SetSharedPreferenceString(String group, String key, String value)
+	{
+		Context context = getApplicationContext();
+		SharedPreferences preferences = context.getSharedPreferences(group, context.MODE_PRIVATE);
 		SharedPreferences.Editor editor = preferences.edit();
-		editor.putInt("AllowCellular", value);
+		editor.putString(key, value);
 		editor.commit();
 	}
 
-	public int AndroidThunkJava_GetCellularPreference()
+	public String AndroidThunkJava_GetSharedPreferenceString(String group, String key, string defaultValue)
 	{
 		Context context = getApplicationContext();
-		SharedPreferences preferences = context.getSharedPreferences("CellularNetworkPreferences", context.MODE_PRIVATE);
-		return preferences.getInt("AllowCellular", 0 );
+		SharedPreferences preferences = context.getSharedPreferences(group, context.MODE_PRIVATE);
+		return preferences.getString(key, defaultValue);
 	}
 
 	// call back into native code from the Java UI thread, initializing any available VR HMD modules
diff --git a/Engine/Source/Runtime/Core/Private/Android/AndroidPlatformMisc.cpp b/Engine/Source/Runtime/Core/Private/Android/AndroidPlatformMisc.cpp
--- a/Engine/Source/Runtime/Core/Private/Android/AndroidPlatformMisc.cpp
+++ b/Engine/Source/Runtime/Core/Private/Android/AndroidPlatformMisc.cpp
@@ -146,9 +146,8 @@ extern void AndroidThunkCpp_ForceQuit();
 
 extern void AndroidThunkCpp_SetOrientation(int32 Value);
 
-extern void AndroidThunkCpp_SetCellularPreference(int32 Value);
-
-extern int32 AndroidThunkCpp_GetCellularPreference();
+extern void AndroidThunkCpp_SetSharedPreferenceInt(const FString& Group, const FString& Key, int32 Value);
+extern int32 AndroidThunkCpp_GetSharedPreferenceInt(const FString& Group, const FString& Key, int32 DefaultValue);
 
 // From AndroidFile.cpp
 extern FString GFontPathBase;
@@ -3125,17 +3124,17 @@ void FAndroidMisc::SetDeviceOrientation(EDeviceScreenOrientation NewDeviceOrenta
 void FAndroidMisc::SetCellularPreference(int32 Value)
 {
 #if USE_ANDROID_JNI
-	AndroidThunkCpp_SetCellularPreference(Value);
+	AndroidThunkCpp_SetSharedPreferenceInt(TEXT("CellularNetworkPreferences"), TEXT("AllowCellular"), Value);
 #endif // USE_ANDROID_JNI
 }
 
 int32 FAndroidMisc::GetCellularPreference()
 {
-	int32 value = 0;
+	int32 Result = 0;
 #if USE_ANDROID_JNI
-	value = AndroidThunkCpp_GetCellularPreference();
+	Result = AndroidThunkCpp_GetSharedPreferenceInt(TEXT("CellularNetworkPreferences"), TEXT("AllowCellular"), Result);
 #endif // USE_ANDROID_JNI
-	return value;
+	return Result;
 }
 
 void FAndroidMisc::SetAllowedDeviceOrientation(EDeviceScreenOrientation NewAllowedDeviceOrientation)
diff --git a/Engine/Source/Runtime/Launch/Private/Android/AndroidJNI.cpp b/Engine/Source/Runtime/Launch/Private/Android/AndroidJNI.cpp
--- a/Engine/Source/Runtime/Launch/Private/Android/AndroidJNI.cpp
+++ b/Engine/Source/Runtime/Launch/Private/Android/AndroidJNI.cpp
@@ -134,8 +134,13 @@ void FJavaWrapper::FindClassesAndMethods(JNIEnv* Env)
 	AndroidThunkJava_GetIntentExtrasString = FindMethod(Env, GameActivityClassID, "AndroidThunkJava_GetIntentExtrasString", "(Ljava/lang/String;)Ljava/lang/String;", bIsOptional);
 	AndroidThunkJava_PushSensorEvents = FindMethod(Env, GameActivityClassID, "AndroidThunkJava_PushSensorEvents", "()V", bIsOptional);
 	AndroidThunkJava_SetOrientation = FindMethod(Env, GameActivityClassID, "AndroidThunkJava_SetOrientation", "(I)V", bIsOptional);
-	AndroidThunkJava_SetCellularPreference = FindMethod(Env, GameActivityClassID, "AndroidThunkJava_SetCellularPreference", "(I)V", bIsOptional);
-	AndroidThunkJava_GetCellularPreference = FindMethod(Env, GameActivityClassID, "AndroidThunkJava_GetCellularPreference", "()I", bIsOptional);
+	AndroidThunkJava_HasSharedPreference = FindMethod(Env, GameActivityClassID, "AndroidThunkJava_HasSharedPreference", "(Ljava/lang/String;Ljava/lang/String;)Z", bIsOptional);
+	AndroidThunkJava_SetSharedPreferenceBoolean = FindMethod(Env, GameActivityClassID, "AndroidThunkJava_SetSharedPreferenceBoolean", "(Ljava/lang/String;Ljava/lang/String;Z)V", bIsOptional);
+	AndroidThunkJava_GetSharedPreferenceBoolean = FindMethod(Env, GameActivityClassID, "AndroidThunkJava_GetSharedPreferenceBoolean", "(Ljava/lang/String;Ljava/lang/String;Z)Z", bIsOptional);
+	AndroidThunkJava_SetSharedPreferenceInt = FindMethod(Env, GameActivityClassID, "AndroidThunkJava_SetSharedPreferenceInt", "(Ljava/lang/String;Ljava/lang/String;I)V", bIsOptional);
+	AndroidThunkJava_GetSharedPreferenceInt = FindMethod(Env, GameActivityClassID, "AndroidThunkJava_GetSharedPreferenceInt", "(Ljava/lang/String;Ljava/lang/String;I)I", bIsOptional);
+	AndroidThunkJava_SetSharedPreferenceString = FindMethod(Env, GameActivityClassID, "AndroidThunkJava_SetSharedPreferenceString", "(Ljava/lang/String;Ljava/lang/String;Ljava/lang/String;)V", bIsOptional);
+	AndroidThunkJava_GetSharedPreferenceString = FindMethod(Env, GameActivityClassID, "AndroidThunkJava_GetSharedPreferenceString", "(Ljava/lang/String;Ljava/lang/String;Ljava/lang/String;)Ljava/lang/String;", bIsOptional);
 
 	// Screen capture/recording permission
 	AndroidThunkJava_IsScreenCaptureDisabled = FindMethod(Env, GameActivityClassID, "AndroidThunkJava_IsScreenCaptureDisabled", "()Z", bIsOptional);
@@ -539,8 +544,13 @@ jmethodID FJavaWrapper::AndroidThunkJava_PushSensorEvents;
 jmethodID FJavaWrapper::AndroidThunkJava_IsScreenCaptureDisabled;
 jmethodID FJavaWrapper::AndroidThunkJava_DisableScreenCapture;
 jmethodID FJavaWrapper::AndroidThunkJava_SetOrientation;
-jmethodID FJavaWrapper::AndroidThunkJava_SetCellularPreference;
-jmethodID FJavaWrapper::AndroidThunkJava_GetCellularPreference;
+jmethodID FJavaWrapper::AndroidThunkJava_HasSharedPreference;
+jmethodID FJavaWrapper::AndroidThunkJava_SetSharedPreferenceBoolean;
+jmethodID FJavaWrapper::AndroidThunkJava_GetSharedPreferenceBoolean;
+jmethodID FJavaWrapper::AndroidThunkJava_SetSharedPreferenceInt;
+jmethodID FJavaWrapper::AndroidThunkJava_GetSharedPreferenceInt;
+jmethodID FJavaWrapper::AndroidThunkJava_SetSharedPreferenceString;
+jmethodID FJavaWrapper::AndroidThunkJava_GetSharedPreferenceString;
 
 jclass FJavaWrapper::InputDeviceInfoClass;
 jfieldID FJavaWrapper::InputDeviceInfo_VendorId;
@@ -1427,22 +1437,84 @@ void AndroidThunkCpp_SetOrientation(int32 Value)
 	}
 }
 
-void AndroidThunkCpp_SetCellularPreference(int32 Value)
+bool AndroidThunkCpp_HasSharedPreference(const FString& Group, const FString& Key)
 {
+	bool Result = false;
+	if (JNIEnv* Env = FAndroidApplication::GetJavaEnv())
+	{
+		auto GroupArg = FJavaHelper::ToJavaString(Env, Group);
+		auto KeyArg = FJavaHelper::ToJavaString(Env, Key);
+		Result = FJavaWrapper::CallBooleanMethod(Env, FJavaWrapper::GameActivityThis, FJavaWrapper::AndroidThunkJava_HasSharedPreference, *GroupArg, *KeyArg);
+	}
+	return Result;
+}
+
+void AndroidThunkCpp_SetSharedPreferenceBool(const FString& Group, const FString& Key, bool Value)
+{
+	if (JNIEnv* Env = FAndroidApplication::GetJavaEnv())
+	{
+		auto GroupArg = FJavaHelper::ToJavaString(Env, Group);
+		auto KeyArg = FJavaHelper::ToJavaString(Env, Key);
+		FJavaWrapper::CallVoidMethod(Env, FJavaWrapper::GameActivityThis, FJavaWrapper::AndroidThunkJava_SetSharedPreferenceBoolean, *GroupArg, *KeyArg, Value);
+	}
+}
+
+bool AndroidThunkCpp_GetSharedPreferenceBool(const FString& Group, const FString& Key, bool DefaultValue)
+{
+	bool Result = DefaultValue;
 	if (JNIEnv* Env = FAndroidApplication::GetJavaEnv())
 	{
-		FJavaWrapper::CallVoidMethod(Env, FJavaWrapper::GameActivityThis, FJavaWrapper::AndroidThunkJava_SetCellularPreference, Value);
+		auto GroupArg = FJavaHelper::ToJavaString(Env, Group);
+		auto KeyArg = FJavaHelper::ToJavaString(Env, Key);
+		Result = FJavaWrapper::CallBooleanMethod(Env, FJavaWrapper::GameActivityThis, FJavaWrapper::AndroidThunkJava_GetSharedPreferenceBoolean, *GroupArg, *KeyArg, DefaultValue);
 	}
+	return Result;
 }
 
-int32 AndroidThunkCpp_GetCellularPreference()
+void AndroidThunkCpp_SetSharedPreferenceInt(const FString& Group, const FString& Key, int32 Value)
 {
-	int32 value = 0;
 	if (JNIEnv* Env = FAndroidApplication::GetJavaEnv())
 	{
-		value = FJavaWrapper::CallIntMethod(Env, FJavaWrapper::GameActivityThis, FJavaWrapper::AndroidThunkJava_GetCellularPreference);
+		auto GroupArg = FJavaHelper::ToJavaString(Env, Group);
+		auto KeyArg = FJavaHelper::ToJavaString(Env, Key);
+		FJavaWrapper::CallVoidMethod(Env, FJavaWrapper::GameActivityThis, FJavaWrapper::AndroidThunkJava_SetSharedPreferenceInt, *GroupArg, *KeyArg, Value);
 	}
-	return value;
+}
+
+int32 AndroidThunkCpp_GetSharedPreferenceInt(const FString& Group, const FString& Key, int32 DefaultValue)
+{
+	int32 Result = DefaultValue;
+	if (JNIEnv* Env = FAndroidApplication::GetJavaEnv())
+	{
+		auto GroupArg = FJavaHelper::ToJavaString(Env, Group);
+		auto KeyArg = FJavaHelper::ToJavaString(Env, Key);
+		Result = FJavaWrapper::CallIntMethod(Env, FJavaWrapper::GameActivityThis, FJavaWrapper::AndroidThunkJava_GetSharedPreferenceInt, *GroupArg, *KeyArg, DefaultValue);
+	}
+	return Result;
+}
+
+void AndroidThunkCpp_SetSharedPreferenceString(const FString& Group, const FString& Key, const FString& Value)
+{
+	if (JNIEnv* Env = FAndroidApplication::GetJavaEnv())
+	{
+		auto GroupArg = FJavaHelper::ToJavaString(Env, Group);
+		auto KeyArg = FJavaHelper::ToJavaString(Env, Key);
+		auto ValueArg = FJavaHelper::ToJavaString(Env, Value);
+		FJavaWrapper::CallVoidMethod(Env, FJavaWrapper::GameActivityThis, FJavaWrapper::AndroidThunkJava_SetSharedPreferenceString, *GroupArg, *KeyArg, *ValueArg);
+	}
+}
+
+FString AndroidThunkCpp_GetSharedPreferenceString(const FString& Group, const FString& Key, const FString& DefaultValue)
+{
+	FString Result = DefaultValue;
+	if (JNIEnv* Env = FAndroidApplication::GetJavaEnv())
+	{
+		auto GroupArg = FJavaHelper::ToJavaString(Env, Group);
+		auto KeyArg = FJavaHelper::ToJavaString(Env, Key);
+		auto DefaultValueArg = FJavaHelper::ToJavaString(Env, DefaultValue);
+		Result = FJavaHelper::FStringFromLocalRef(Env, (jstring)FJavaWrapper::CallObjectMethod(Env, FJavaWrapper::GameActivityThis, FJavaWrapper::AndroidThunkJava_GetSharedPreferenceString, *GroupArg, *KeyArg, *DefaultValueArg));
+	}
+	return Result;
 }
 
 bool AndroidThunkCpp_IsMusicActive()
diff --git a/Engine/Source/Runtime/Launch/Public/Android/AndroidJNI.h b/Engine/Source/Runtime/Launch/Public/Android/AndroidJNI.h
--- a/Engine/Source/Runtime/Launch/Public/Android/AndroidJNI.h
+++ b/Engine/Source/Runtime/Launch/Public/Android/AndroidJNI.h
@@ -73,8 +73,13 @@ public:
 	static jmethodID AndroidThunkJava_SetSustainedPerformanceMode;
 	static jmethodID AndroidThunkJava_PushSensorEvents;
 	static jmethodID AndroidThunkJava_SetOrientation;
-	static jmethodID AndroidThunkJava_SetCellularPreference;
-	static jmethodID AndroidThunkJava_GetCellularPreference;
+	static jmethodID AndroidThunkJava_HasSharedPreference;
+	static jmethodID AndroidThunkJava_SetSharedPreferenceBoolean;
+	static jmethodID AndroidThunkJava_GetSharedPreferenceBoolean;
+	static jmethodID AndroidThunkJava_SetSharedPreferenceInt;
+	static jmethodID AndroidThunkJava_GetSharedPreferenceInt;
+	static jmethodID AndroidThunkJava_SetSharedPreferenceString;
+	static jmethodID AndroidThunkJava_GetSharedPreferenceString;
 
 
 	// Screen capture/recording permission
