Add new EEventThread tag. It is using for Android event thread

diff --git a/Engine/Source/Runtime/Core/Public/CoreGlobals.h b/Engine/Source/Runtime/Core/Public/CoreGlobals.h
--- a/Engine/Source/Runtime/Core/Public/CoreGlobals.h
+++ b/Engine/Source/Runtime/Core/Public/CoreGlobals.h
@@ -649,10 +649,11 @@ enum class ETaskTag : int32
 	ERenderingThread			= 1 << 4,
 	ERhiThread					= 1 << 5,
 	EAsyncLoadingThread			= 1 << 6,
+	EEventThread				= 1 << 7,
 
-	ENamedThreadBits			= (EAsyncLoadingThread << 1) - 1,
-	EParallelThread				= 1 << 8, //This can be used when multipe threads or jobs are involved (usually a parallel for) It will avoid the check for uniqieness of the named thread tag.
-	EWorkerThread				= 1 << 7 | EParallelThread,
+	ENamedThreadBits			= (EEventThread << 1) - 1,
+	EParallelThread				= 1 << 30, //This can be used when multiple threads or jobs are involved (usually a parallel for) It will avoid the check for uniqueness of the named thread tag.
+	EWorkerThread				= 1 << 29 | EParallelThread,
 	EParallelRenderingThread	= ERenderingThread | EParallelThread,
 	EParallelGameThread			= EGameThread | EParallelThread,
 	EParallelRhiThread			= ERhiThread | EParallelThread,
diff --git a/Engine/Source/Runtime/Launch/Private/Android/LaunchAndroid.cpp b/Engine/Source/Runtime/Launch/Private/Android/LaunchAndroid.cpp
--- a/Engine/Source/Runtime/Launch/Private/Android/LaunchAndroid.cpp
+++ b/Engine/Source/Runtime/Launch/Private/Android/LaunchAndroid.cpp
@@ -524,6 +524,7 @@ static void ApplyAndroidCompatConfigRules()
 //Main function called from the android entry point
 int32 AndroidMain(struct android_app* state);
 
+// The function is used only for ASIS
 void* AndroidMain(void* param)
 {
 	struct android_app* state = (struct android_app*)param;
@@ -535,11 +536,6 @@ void* AndroidMain(void* param)
 		STANDALONE_DEBUG_LOGf(LogAndroid, TEXT("AndroidMain set current GGameThreadId=%d"), GGameThreadId);
 	}
 
-	if (EventThreadID == 0)
-	{
-		EventThreadID = GGameThreadId;
-		STANDALONE_DEBUG_LOGf(LogAndroid, TEXT("AndroidMain set current EventThreadID=%d"), EventThreadID);
-	}
 	AndroidMain(state);
 	return nullptr;
 }
@@ -901,17 +897,13 @@ bool IsInAndroidEventThread()
 	// Note: leave the commented out line for debug purposes.
 	//STANDALONE_DEBUG_LOGf(LogAndroid, TEXT("IsInAndroidEventThread(), GGameThreadId=%d, EventThreadID=%d, FPlatformTLS::GetCurrentThreadId()=%d"), GGameThreadId, EventThreadID, FPlatformTLS::GetCurrentThreadId());
 
-#if USE_ANDROID_STANDALONE
-	//@TODO: for now always return true to avoid check failures
-	return true;
-#else
 	check(EventThreadID != 0);
 	return EventThreadID == FPlatformTLS::GetCurrentThreadId();
-#endif
 }
 
 static void* AndroidEventThreadWorker( void* param )
 {
+	FTaskTagScope Scope(ETaskTag::EEventThread);
 	pthread_setname_np(pthread_self(), "EventWorker");
 	STANDALONE_DEBUG_LOGf(LogAndroid, TEXT("AndroidEventThreadWorker(begin), GGameThreadId=%d, EventThreadID=%d, FPlatformTLS::GetCurrentThreadId()=%d"), GGameThreadId, EventThreadID, FPlatformTLS::GetCurrentThreadId());
 
@@ -2161,6 +2153,13 @@ JNI_METHOD void Java_com_epicgames_makeaar_GameActivityForMakeAAR_nativeMain(JNI
 
 	STANDALONE_DEBUG_LOGf(LogAndroid, TEXT("Java_com_epicgames_makeaar_GameActivityForMakeAAR_nativeMain : use current!, requesting ProjectModule: %s"), *projectModuleName);
 
+	if (EventThreadID == 0)
+	{
+		FTaskTagScope::SwapTag(ETaskTag::EEventThread);
+		EventThreadID = FPlatformTLS::GetCurrentThreadId();
+		STANDALONE_DEBUG_LOGf(LogAndroid, TEXT("AndroidMain set current EventThreadID=%d"), EventThreadID);
+	}
+
 	// register some delegates we want to pass back
 	FCoreDelegates::OnInit.AddStatic(InitEvent);
 	FCoreDelegates::OnPostEngineInit.AddStatic(PostEngineInitEvent);
