Add support for NDK 27

diff --git a/Engine/Config/Android/Android_SDK.json b/Engine/Config/Android/Android_SDK.json
--- a/Engine/Config/Android/Android_SDK.json
+++ b/Engine/Config/Android/Android_SDK.json
@@ -1,7 +1,7 @@
 {
 	"MainVersion": "r25b",
 	"MinVersion": "r25b",
-	"MaxVersion": "r26c",
+	"MaxVersion": "r27",
 
 	"AutoSDKDirectory": "-26",
 
diff --git a/Engine/Source/Runtime/CoreUObject/Private/UObject/GarbageCollection.cpp b/Engine/Source/Runtime/CoreUObject/Private/UObject/GarbageCollection.cpp
--- a/Engine/Source/Runtime/CoreUObject/Private/UObject/GarbageCollection.cpp
+++ b/Engine/Source/Runtime/CoreUObject/Private/UObject/GarbageCollection.cpp
@@ -1937,8 +1937,8 @@ private:
 		{
 			return;
 		}
-		
-		uint32 ValidIndices[UnvalidatedStructArrays.Max()];
+
+		uint32 ValidIndices[decltype(UnvalidatedStructArrays)::Max()];
 		uint32* ValidIt = ValidIndices;
 		for (uint32 Idx = 0; Idx < Num; ++Idx)
 		{
diff --git a/Engine/Source/Runtime/Launch/Private/Android/LaunchAndroid.cpp b/Engine/Source/Runtime/Launch/Private/Android/LaunchAndroid.cpp
--- a/Engine/Source/Runtime/Launch/Private/Android/LaunchAndroid.cpp
+++ b/Engine/Source/Runtime/Launch/Private/Android/LaunchAndroid.cpp
@@ -203,9 +203,6 @@ static void AndroidProcessEvents(struct android_app* state);
 //Event thread stuff
 static void* AndroidEventThreadWorker(void* param);
 
-// How often to process (read & dispatch) events, in seconds.
-static const float EventRefreshRate = 1.0f / 20.0f;
-
 // Name of the UE commandline append setprop
 static constexpr char UECommandLineSetprop[] = "debug.ue.commandline";
 
@@ -943,11 +940,7 @@ static void* AndroidEventThreadWorker( void* param )
 	//continue to process events until the engine is shutting down
 	while (!IsEngineExitRequested())
 	{
-		//		FPlatformMisc::LowLevelOutputDebugString(TEXT("AndroidEventThreadWorker"));
-
 		AndroidProcessEvents(state);
-
-		sleep(EventRefreshRate);		// this is really 0 since it takes int seconds.
 	}
 	DEVELOPER_LOG_COMMANDCB_CASE(AndroidEventThreadWorker_AfterWhile);
 
@@ -960,20 +953,14 @@ static void* AndroidEventThreadWorker( void* param )
 }
 
 //Called from the separate event processing thread
-static void AndroidProcessEvents(struct android_app* state)
+static void AndroidProcessEvents(struct android_app* State)
 {
-	int ident;
-	int fdesc;
-	int events;
-	struct android_poll_source* source;
+	struct android_poll_source* Source = nullptr;
+	int32 Result = ALooper_pollOnce(-1, nullptr, nullptr, (void**)&Source);
 
-	while ((ident = ALooper_pollAll(-1, &fdesc, &events, (void**)&source)) >= 0)
+	if (Result != ALOOPER_POLL_ERROR && Source != nullptr)
 	{
-		// process this event
-		if (source)
-		{
-			source->process(state, source);
-		}
+		Source->process(State, Source);
 	}
 }
 
