Use SwapTag to set EEventThread tag for Android event thread

#rnx

diff --git a/Engine/Source/Runtime/Launch/Private/Android/LaunchAndroid.cpp b/Engine/Source/Runtime/Launch/Private/Android/LaunchAndroid.cpp
--- a/Engine/Source/Runtime/Launch/Private/Android/LaunchAndroid.cpp
+++ b/Engine/Source/Runtime/Launch/Private/Android/LaunchAndroid.cpp
@@ -900,7 +900,7 @@ bool IsInAndroidEventThread()
 
 static void* AndroidEventThreadWorker( void* param )
 {
-	FTaskTagScope Scope(ETaskTag::EEventThread);
+	FTaskTagScope::SwapTag(ETaskTag::EEventThread);
 	pthread_setname_np(pthread_self(), "EventWorker");
 	STANDALONE_DEBUG_LOGf(LogAndroid, TEXT("AndroidEventThreadWorker(begin), GGameThreadId=%d, EventThreadID=%d, FPlatformTLS::GetCurrentThreadId()=%d"), GGameThreadId, EventThreadID, FPlatformTLS::GetCurrentThreadId());
 
