Strip "/r" and "/n" from end of line for any string sent to iOS's TcpConsoleListener.

#rnx

diff --git a/Engine/Source/Runtime/Launch/Private/IOS/TcpConsoleListener.cpp b/Engine/Source/Runtime/Launch/Private/IOS/TcpConsoleListener.cpp
--- a/Engine/Source/Runtime/Launch/Private/IOS/TcpConsoleListener.cpp
+++ b/Engine/Source/Runtime/Launch/Private/IOS/TcpConsoleListener.cpp
@@ -80,6 +80,12 @@ uint32 TcpConsoleListener::Run()
 				memset(RecvBuffer, 0, CommandSize);
 				if (Connection->Recv(RecvBuffer, CommandSize, BytesRead))
 				{
+					// strip any carriage return/newline character from end of line
+					int32 BufferIdx = BytesRead - 1;
+					while (BufferIdx >= 0 && (RecvBuffer[BufferIdx] == '\n' || RecvBuffer[BufferIdx] == '\r'))
+					{
+						RecvBuffer[BufferIdx--] = 0;
+					}
 					UE_LOG(LogTemp, Display, TEXT("Received TCP console command '%s'"), *FString(UTF8_TO_TCHAR(RecvBuffer)));
 					EnqueueConsoleCommand(RecvBuffer);
 				}
