Fix issue with HMD initialization on Android

#12055
#android

diff --git a/Engine/Source/Runtime/Launch/Private/Android/LaunchAndroid.cpp b/Engine/Source/Runtime/Launch/Private/Android/LaunchAndroid.cpp
--- a/Engine/Source/Runtime/Launch/Private/Android/LaunchAndroid.cpp
+++ b/Engine/Source/Runtime/Launch/Private/Android/LaunchAndroid.cpp
@@ -337,11 +337,13 @@ void InitHMDs()
 
 	// Get a list of plugins that implement this feature
 	GHMDImplementations = IModularFeatures::Get().GetModularFeatureImplementations<IHeadMountedDisplayModule>(IHeadMountedDisplayModule::GetModularFeatureName());
+	if (GHMDImplementations.IsEmpty())
+	{
+		return;
+	}
 
 	AndroidThunkCpp_InitHMDs();
 
-	GHMDsInitialized = GHMDImplementations.IsEmpty();
-
 	while (!GHMDsInitialized)
 	{
 		FPlatformProcess::Sleep(0.01f);
