Fix potentially longer startup times in certain non-monolithic platform builds due to plugin discovery phase occurring prior to task graph system initialization.

#rnx

diff --git a/Engine/Source/Runtime/Launch/Private/LaunchEngineLoop.cpp b/Engine/Source/Runtime/Launch/Private/LaunchEngineLoop.cpp
--- a/Engine/Source/Runtime/Launch/Private/LaunchEngineLoop.cpp
+++ b/Engine/Source/Runtime/Launch/Private/LaunchEngineLoop.cpp
@@ -2687,14 +2687,6 @@ int32 FEngineLoop::PreInitPreStartupScreen(const TCHAR* CmdLine)
 		FPlatformProcess::AddDllDirectory(*ProjectBinariesDirectory);
 		FModuleManager::Get().SetGameBinariesDirectory(*ProjectBinariesDirectory);
 
-#if !IS_MONOLITHIC
-		if (FCString::Strcmp(CurrentPlatformFile->GetName(), TEXT("PakFile")) == 0)
-		{
-			IPluginManager::Get().SetBinariesRootDirectories(EngineBinariesRootDirectory, ProjectBinariesRootDirectory);
-			IPluginManager::Get().SetPreloadBinaries();
-		}
-#endif
-
 		LaunchFixGameNameCase();
 	}
 #endif
@@ -2797,6 +2789,14 @@ int32 FEngineLoop::PreInitPreStartupScreen(const TCHAR* CmdLine)
 	// this can start using TaskGraph and ThreadPool so they must be created before
 	FDelayedAutoRegisterHelper::RunAndClearDelayedAutoRegisterDelegates(EDelayedRegisterRunPhase::TaskGraphSystemReady);
 
+#if !IS_PROGRAM && !IS_MONOLITHIC
+	if (FApp::HasProjectName() && FCString::Strcmp(CurrentPlatformFile->GetName(), TEXT("PakFile")) == 0)
+	{
+		IPluginManager::Get().SetBinariesRootDirectories(EngineBinariesRootDirectory, ProjectBinariesRootDirectory);
+		IPluginManager::Get().SetPreloadBinaries();
+	}
+#endif
+
 #if STATS
 	FThreadStats::StartThread();
 #endif
