Fix for rendering occuring when moving from the background to the foreground before the application is ready.

diff --git a/Engine/Source/Runtime/Launch/Private/IOS/LaunchIOS.cpp b/Engine/Source/Runtime/Launch/Private/IOS/LaunchIOS.cpp
--- a/Engine/Source/Runtime/Launch/Private/IOS/LaunchIOS.cpp
+++ b/Engine/Source/Runtime/Launch/Private/IOS/LaunchIOS.cpp
@@ -57,6 +57,8 @@ void FAppEntry::Suspend(bool bIsInterrupt)
 	{
 		GetMoviePlayer()->Suspend();
 	}
+	
+	FPreLoadScreenManager::EnableRendering(false);
 
 	// if background audio is active, then we don't want to do suspend any audio
 	if ([[IOSAppDelegate GetDelegate] IsFeatureActive:EAudioFeature::BackgroundAudio] == false)
@@ -150,6 +152,8 @@ void FAppEntry::Resume(bool bIsInterrupt)
 	{
 		GetMoviePlayer()->Resume();
 	}
+	
+	FPreLoadScreenManager::EnableRendering(true);
 
 	// if background audio is active, then we don't want to do suspend any audio
 	// @todo: should this check if we were suspended, in case this changes while in the background? (suspend with background off, but resume with background audio on? is that a thing?)
