Removed legacy manual sync of oustanding command lists to delay scene render completion. Scene render clean up is handled as a separate task graph that absorbs those command lists. This was previously not doing anything but now blocks async RDG execution tasks from running.

diff --git a/Engine/Source/Runtime/Launch/Private/LaunchEngineLoop.cpp b/Engine/Source/Runtime/Launch/Private/LaunchEngineLoop.cpp
--- a/Engine/Source/Runtime/Launch/Private/LaunchEngineLoop.cpp
+++ b/Engine/Source/Runtime/Launch/Private/LaunchEngineLoop.cpp
@@ -5991,14 +5991,6 @@ void FEngineLoop::Tick()
 			ConcurrentTaskCompleteEvent = nullptr;
 			ConcurrentTask = nullptr;
 		}
-		{
-			ENQUEUE_RENDER_COMMAND(WaitForOutstandingTasksOnly_for_DelaySceneRenderCompletion)(
-				[](FRHICommandListImmediate& RHICmdList)
-				{
-					QUICK_SCOPE_CYCLE_COUNTER(STAT_DelaySceneRenderCompletion_TaskWait);
-					RHICmdList.ImmediateFlush(EImmediateFlushType::WaitForOutstandingTasksOnly);
-				});
-		}
 #endif
 
 #if STATS
