Implement FPlatformMisc::SetStoredValue, GetStoredValue and others on Android

diff --git a/Engine/Build/Android/Java/src/com/epicgames/unreal/GameActivity.java.template b/Engine/Build/Android/Java/src/com/epicgames/unreal/GameActivity.java.template
--- a/Engine/Build/Android/Java/src/com/epicgames/unreal/GameActivity.java.template
+++ b/Engine/Build/Android/Java/src/com/epicgames/unreal/GameActivity.java.template
@@ -4816,6 +4816,24 @@ public class GameActivity extends $${gameActivitySuperClass}$$ implements Surfac
 		return preferences.getString(key, defaultValue);
 	}
 
+	public void AndroidThunkJava_DeleteSharedPreference(String group, String key)
+	{
+		Context context = getApplicationContext();
+		SharedPreferences preferences = context.getSharedPreferences(group, context.MODE_PRIVATE);
+		SharedPreferences.Editor editor = preferences.edit();
+		editor.remove(key);
+		editor.commit();
+	}
+
+	public void AndroidThunkJava_DeleteSharedPreferenceGroup(String group)
+	{
+		Context context = getApplicationContext();
+		SharedPreferences preferences = context.getSharedPreferences(group, context.MODE_PRIVATE);
+		SharedPreferences.Editor editor = preferences.edit();
+		editor.clear();
+		editor.commit();
+	}
+
 	// call back into native code from the Java UI thread, initializing any available VR HMD modules
 	public void AndroidThunkJava_InitHMDs()
 	{
diff --git a/Engine/Source/Runtime/Core/Private/Android/AndroidPlatformMisc.cpp b/Engine/Source/Runtime/Core/Private/Android/AndroidPlatformMisc.cpp
--- a/Engine/Source/Runtime/Core/Private/Android/AndroidPlatformMisc.cpp
+++ b/Engine/Source/Runtime/Core/Private/Android/AndroidPlatformMisc.cpp
@@ -146,8 +146,15 @@ extern void AndroidThunkCpp_ForceQuit();
 
 extern void AndroidThunkCpp_SetOrientation(int32 Value);
 
+#if USE_ANDROID_JNI
+extern bool AndroidThunkCpp_HasSharedPreference(const FString& Group, const FString& Key);
 extern void AndroidThunkCpp_SetSharedPreferenceInt(const FString& Group, const FString& Key, int32 Value);
 extern int32 AndroidThunkCpp_GetSharedPreferenceInt(const FString& Group, const FString& Key, int32 DefaultValue);
+extern void AndroidThunkCpp_SetSharedPreferenceString(const FString& Group, const FString& Key, const FString& Value);
+extern bool AndroidThunkCpp_GetSharedPreferenceStringTypeSafe(const FString& Group, const FString& Key, FString& OutValue);
+extern void AndroidThunkCpp_DeleteSharedPreference(const FString& Group, const FString& Key);
+extern void AndroidThunkCpp_DeleteSharedPreferenceGroup(const FString& Group);
+#endif
 
 // From AndroidFile.cpp
 extern FString GFontPathBase;
@@ -788,6 +795,70 @@ bool FAndroidMisc::HasPlatformFeature(const TCHAR* FeatureName)
 	return FGenericPlatformMisc::HasPlatformFeature(FeatureName);
 }
 
+static FString FixUpStoredValueSectionName(const FString& InSectionName)
+{
+	// We need to remove file separator character because SectionName is used as a file name by the OS
+	if (InSectionName.Contains(TEXT("/")))
+	{
+		// Use extended replacement to avoid foo/bar and foo_bar mapping to each other
+		return InSectionName.Replace(TEXT("/"), TEXT("___"));
+	}
+	else
+	{
+		return InSectionName;
+	}
+}
+
+bool FAndroidMisc::SetStoredValue(const FString& InStoreId, const FString& InSectionName, const FString& InKeyName, const FString& InValue)
+{
+#if USE_ANDROID_JNI
+	const FString FixedSectionValue = FixUpStoredValueSectionName(InSectionName);
+	AndroidThunkCpp_SetSharedPreferenceString(FixedSectionValue, InKeyName, InValue);
+	return true;
+#else
+	return FGenericPlatformMisc::SetStoredValue(InStoreId, InSectionName, InKeyName, InValue);
+#endif
+}
+
+bool FAndroidMisc::GetStoredValue(const FString& InStoreId, const FString& InSectionName, const FString& InKeyName, FString& OutValue)
+{
+#if USE_ANDROID_JNI
+	const FString FixedSectionValue = FixUpStoredValueSectionName(InSectionName);
+	if (AndroidThunkCpp_GetSharedPreferenceStringTypeSafe(FixedSectionValue, InKeyName, OutValue))
+	{
+		return true;
+	}
+#endif
+
+	// fallback to FGenericPlatformMisc::GetStoredValue
+	return FGenericPlatformMisc::GetStoredValue(InStoreId, InSectionName, InKeyName, OutValue);
+}
+
+bool FAndroidMisc::DeleteStoredValue(const FString& InStoreId, const FString& InSectionName, const FString& InKeyName)
+{
+	bool bResult = false;
+#if USE_ANDROID_JNI
+	const FString FixedSectionValue = FixUpStoredValueSectionName(InSectionName);
+	// delete doesn't have a return value, our best effort is to check for preference existence first
+	bResult = AndroidThunkCpp_HasSharedPreference(FixedSectionValue, InKeyName);
+	AndroidThunkCpp_DeleteSharedPreference(FixedSectionValue, InKeyName);
+#endif
+	// always delete in both places just in case
+	bResult |= FGenericPlatformMisc::DeleteStoredValue(InStoreId, InSectionName, InKeyName);
+	return bResult;
+}
+
+bool FAndroidMisc::DeleteStoredSection(const FString& InStoreId, const FString& InSectionName)
+{
+#if USE_ANDROID_JNI
+	const FString FixedSectionValue = FixUpStoredValueSectionName(InSectionName);
+	// can't easily check if section exists as Context.getSharedPreferences will create object if it doesn't exist
+	AndroidThunkCpp_DeleteSharedPreferenceGroup(FixedSectionValue);
+#endif
+	// always delete in both places just in case
+	return FGenericPlatformMisc::DeleteStoredSection(InStoreId, InSectionName);
+}
+
 bool FAndroidMisc::UseRenderThread()
 {
 	// if we in general don't want to use the render thread due to commandline, etc, then don't
diff --git a/Engine/Source/Runtime/Core/Public/Android/AndroidPlatformMisc.h b/Engine/Source/Runtime/Core/Public/Android/AndroidPlatformMisc.h
--- a/Engine/Source/Runtime/Core/Public/Android/AndroidPlatformMisc.h
+++ b/Engine/Source/Runtime/Core/Public/Android/AndroidPlatformMisc.h
@@ -59,6 +59,10 @@ struct FAndroidMisc : public FGenericPlatformMisc
 	static CORE_API bool UseRenderThread();
 	static CORE_API bool HasPlatformFeature(const TCHAR* FeatureName);
 	static CORE_API bool SupportsES30();
+	static CORE_API bool SetStoredValue(const FString& InStoreId, const FString& InSectionName, const FString& InKeyName, const FString& InValue);
+	static CORE_API bool GetStoredValue(const FString& InStoreId, const FString& InSectionName, const FString& InKeyName, FString& OutValue);
+	static CORE_API bool DeleteStoredValue(const FString& InStoreId, const FString& InSectionName, const FString& InKeyName);
+	static CORE_API bool DeleteStoredSection(const FString& InStoreId, const FString& InSectionName);
 
 public:
 
diff --git a/Engine/Source/Runtime/Launch/Private/Android/AndroidJNI.cpp b/Engine/Source/Runtime/Launch/Private/Android/AndroidJNI.cpp
--- a/Engine/Source/Runtime/Launch/Private/Android/AndroidJNI.cpp
+++ b/Engine/Source/Runtime/Launch/Private/Android/AndroidJNI.cpp
@@ -141,6 +141,8 @@ void FJavaWrapper::FindClassesAndMethods(JNIEnv* Env)
 	AndroidThunkJava_GetSharedPreferenceInt = FindMethod(Env, GameActivityClassID, "AndroidThunkJava_GetSharedPreferenceInt", "(Ljava/lang/String;Ljava/lang/String;I)I", bIsOptional);
 	AndroidThunkJava_SetSharedPreferenceString = FindMethod(Env, GameActivityClassID, "AndroidThunkJava_SetSharedPreferenceString", "(Ljava/lang/String;Ljava/lang/String;Ljava/lang/String;)V", bIsOptional);
 	AndroidThunkJava_GetSharedPreferenceString = FindMethod(Env, GameActivityClassID, "AndroidThunkJava_GetSharedPreferenceString", "(Ljava/lang/String;Ljava/lang/String;Ljava/lang/String;)Ljava/lang/String;", bIsOptional);
+	AndroidThunkJava_DeleteSharedPreference = FindMethod(Env, GameActivityClassID, "AndroidThunkJava_DeleteSharedPreference", "(Ljava/lang/String;Ljava/lang/String;)V", bIsOptional);
+	AndroidThunkJava_DeleteSharedPreferenceGroup = FindMethod(Env, GameActivityClassID, "AndroidThunkJava_DeleteSharedPreferenceGroup", "(Ljava/lang/String;)V", bIsOptional);
 
 	// Screen capture/recording permission
 	AndroidThunkJava_IsScreenCaptureDisabled = FindMethod(Env, GameActivityClassID, "AndroidThunkJava_IsScreenCaptureDisabled", "()Z", bIsOptional);
@@ -551,6 +553,8 @@ jmethodID FJavaWrapper::AndroidThunkJava_SetSharedPreferenceInt;
 jmethodID FJavaWrapper::AndroidThunkJava_GetSharedPreferenceInt;
 jmethodID FJavaWrapper::AndroidThunkJava_SetSharedPreferenceString;
 jmethodID FJavaWrapper::AndroidThunkJava_GetSharedPreferenceString;
+jmethodID FJavaWrapper::AndroidThunkJava_DeleteSharedPreference;
+jmethodID FJavaWrapper::AndroidThunkJava_DeleteSharedPreferenceGroup;
 
 jclass FJavaWrapper::InputDeviceInfoClass;
 jfieldID FJavaWrapper::InputDeviceInfo_VendorId;
@@ -1471,6 +1475,30 @@ bool AndroidThunkCpp_GetSharedPreferenceBool(const FString& Group, const FString
 	return Result;
 }
 
+bool AndroidThunkCpp_GetSharedPreferenceBoolTypeSafe(const FString& Group, const FString& Key, bool& OutValue)
+{
+	if (JNIEnv* Env = FAndroidApplication::GetJavaEnv())
+	{
+		auto GroupArg = FJavaHelper::ToJavaString(Env, Group);
+		auto KeyArg = FJavaHelper::ToJavaString(Env, Key);
+		bool bResult = FJavaWrapper::CallBooleanMethod(Env, FJavaWrapper::GameActivityThis, FJavaWrapper::AndroidThunkJava_GetSharedPreferenceBoolean, *GroupArg, *KeyArg, false);
+
+		// SharedPreferences.getBoolean will throw an exception if there is a type mismatch
+		if (!Env->ExceptionCheck())
+		{
+			OutValue = bResult;
+			return true;
+		}
+		else
+		{
+			Env->ExceptionClear();
+			return false;
+		}
+	}
+
+	return false;
+}
+
 void AndroidThunkCpp_SetSharedPreferenceInt(const FString& Group, const FString& Key, int32 Value)
 {
 	if (JNIEnv* Env = FAndroidApplication::GetJavaEnv())
@@ -1493,6 +1521,30 @@ int32 AndroidThunkCpp_GetSharedPreferenceInt(const FString& Group, const FString
 	return Result;
 }
 
+bool AndroidThunkCpp_GetSharedPreferenceIntTypeSafe(const FString& Group, const FString& Key, int32& OutValue)
+{
+	if (JNIEnv* Env = FAndroidApplication::GetJavaEnv())
+	{
+		auto GroupArg = FJavaHelper::ToJavaString(Env, Group);
+		auto KeyArg = FJavaHelper::ToJavaString(Env, Key);
+		int32 Result = FJavaWrapper::CallIntMethod(Env, FJavaWrapper::GameActivityThis, FJavaWrapper::AndroidThunkJava_GetSharedPreferenceInt, *GroupArg, *KeyArg, 0);
+
+		// SharedPreferences.getInt will throw an exception if there is a type mismatch
+		if (!Env->ExceptionCheck())
+		{
+			OutValue = Result;
+			return true;
+		}
+		else
+		{
+			Env->ExceptionClear();
+			return false;
+		}
+	}
+
+	return false;
+}
+
 void AndroidThunkCpp_SetSharedPreferenceString(const FString& Group, const FString& Key, const FString& Value)
 {
 	if (JNIEnv* Env = FAndroidApplication::GetJavaEnv())
@@ -1517,6 +1569,49 @@ FString AndroidThunkCpp_GetSharedPreferenceString(const FString& Group, const FS
 	return Result;
 }
 
+bool AndroidThunkCpp_GetSharedPreferenceStringTypeSafe(const FString& Group, const FString& Key, FString& OutValue)
+{
+	if (JNIEnv* Env = FAndroidApplication::GetJavaEnv())
+	{
+		auto GroupArg = FJavaHelper::ToJavaString(Env, Group);
+		auto KeyArg = FJavaHelper::ToJavaString(Env, Key);
+		auto ResultObject = FJavaWrapper::CallObjectMethod(Env, FJavaWrapper::GameActivityThis, FJavaWrapper::AndroidThunkJava_GetSharedPreferenceString, *GroupArg, *KeyArg, NULL);
+
+		// SharedPreferences.getString will throw an exception if there is a type mismatch
+		if (Env->ExceptionCheck())
+		{
+			Env->ExceptionClear();
+			return false;
+		}
+		else if (ResultObject && !Env->IsSameObject(ResultObject, NULL))
+		{
+			OutValue = FJavaHelper::FStringFromLocalRef(Env, (jstring)ResultObject);
+			return true;
+		}
+	}
+
+	return false;
+}
+
+void AndroidThunkCpp_DeleteSharedPreference(const FString& Group, const FString& Key)
+{
+	if (JNIEnv* Env = FAndroidApplication::GetJavaEnv())
+	{
+		auto GroupArg = FJavaHelper::ToJavaString(Env, Group);
+		auto KeyArg = FJavaHelper::ToJavaString(Env, Key);
+		FJavaWrapper::CallVoidMethod(Env, FJavaWrapper::GameActivityThis, FJavaWrapper::AndroidThunkJava_DeleteSharedPreference, *GroupArg, *KeyArg);
+	}
+}
+
+void AndroidThunkCpp_DeleteSharedPreferenceGroup(const FString& Group)
+{
+	if (JNIEnv* Env = FAndroidApplication::GetJavaEnv())
+	{
+		auto GroupArg = FJavaHelper::ToJavaString(Env, Group);
+		FJavaWrapper::CallVoidMethod(Env, FJavaWrapper::GameActivityThis, FJavaWrapper::AndroidThunkJava_DeleteSharedPreferenceGroup, *GroupArg);
+	}
+}
+
 bool AndroidThunkCpp_IsMusicActive()
 {
 	bool bIsActive = false;
diff --git a/Engine/Source/Runtime/Launch/Public/Android/AndroidJNI.h b/Engine/Source/Runtime/Launch/Public/Android/AndroidJNI.h
--- a/Engine/Source/Runtime/Launch/Public/Android/AndroidJNI.h
+++ b/Engine/Source/Runtime/Launch/Public/Android/AndroidJNI.h
@@ -80,6 +80,8 @@ public:
 	static jmethodID AndroidThunkJava_GetSharedPreferenceInt;
 	static jmethodID AndroidThunkJava_SetSharedPreferenceString;
 	static jmethodID AndroidThunkJava_GetSharedPreferenceString;
+	static jmethodID AndroidThunkJava_DeleteSharedPreference;
+	static jmethodID AndroidThunkJava_DeleteSharedPreferenceGroup;
 
 
 	// Screen capture/recording permission
