Fix secure coding runtime warning

#rnx

diff --git a/Engine/Source/Runtime/Launch/Private/Mac/LaunchMac.cpp b/Engine/Source/Runtime/Launch/Private/Mac/LaunchMac.cpp
--- a/Engine/Source/Runtime/Launch/Private/Mac/LaunchMac.cpp
+++ b/Engine/Source/Runtime/Launch/Private/Mac/LaunchMac.cpp
@@ -423,6 +423,11 @@ static int32 MacOSVersionCompare(const NSOperatingSystemVersion& VersionA, const
 	RunGameThread(self, @selector(runGameThread:));
 }
 
+- (BOOL)applicationSupportsSecureRestorableState:(NSApplication *)app
+{
+    return false;
+}
+
 @end
 
 extern bool GIsConsoleExecutable;
