Fixed dependency of Launch module to deprecated ProfilerClient and ProfilerService modules (toggled by UE_DEPRECATED_PROFILER_ENABLED).

diff --git a/Engine/Source/Runtime/Launch/Launch.Build.cs b/Engine/Source/Runtime/Launch/Launch.Build.cs
--- a/Engine/Source/Runtime/Launch/Launch.Build.cs
+++ b/Engine/Source/Runtime/Launch/Launch.Build.cs
@@ -1,7 +1,9 @@
 // Copyright Epic Games, Inc. All Rights Reserved.
 
-using UnrealBuildTool;
+using System;
 using System.IO;
+using System.Linq;
+using UnrealBuildTool;
 
 public class Launch : ModuleRules
 {
@@ -166,13 +168,16 @@ public class Launch : ModuleRules
 
 		if (Target.Configuration != UnrealTargetConfiguration.Shipping)
 		{
-			PublicIncludePathModuleNames.Add("ProfilerService");
-
 			DynamicallyLoadedModuleNames.AddRange(new string[] {
 				"ProfileVisualizer",
 				"RealtimeProfiler",
-				"ProfilerService"
 			});
+
+			if (Target.GlobalDefinitions.Contains("UE_DEPRECATED_PROFILER_ENABLED=1"))
+			{
+				PublicIncludePathModuleNames.Add("ProfilerService");
+				DynamicallyLoadedModuleNames.Add("ProfilerService");
+			}
 		}
 
 		// The engine can use AutomationController in any connfiguration besides shipping.  This module is loaded
@@ -184,7 +189,11 @@ public class Launch : ModuleRules
 
 		if (Target.bBuildEditor == true)
 		{
-			PublicIncludePathModuleNames.Add("ProfilerClient");
+			if (Target.GlobalDefinitions.Contains("UE_DEPRECATED_PROFILER_ENABLED=1"))
+			{
+				PublicIncludePathModuleNames.Add("ProfilerClient");
+				DynamicallyLoadedModuleNames.Add("ProfilerClient");
+			}
 
 			PrivateDependencyModuleNames.AddRange(new string[] {
 					"SourceControl",
@@ -194,11 +203,9 @@ public class Launch : ModuleRules
 					"PIEPreviewDeviceProfileSelector",
 			});
 
-
 			// ExtraModules that are loaded when WITH_EDITOR=1 is true
 			DynamicallyLoadedModuleNames.AddRange(new string[] {
 					"AutomationWindow",
-					"ProfilerClient",
 					"OutputLog",
 					"TextureCompressor",
 					"MeshUtilities",
diff --git a/Engine/Source/Runtime/Launch/Private/LaunchEngineLoop.cpp b/Engine/Source/Runtime/Launch/Private/LaunchEngineLoop.cpp
--- a/Engine/Source/Runtime/Launch/Private/LaunchEngineLoop.cpp
+++ b/Engine/Source/Runtime/Launch/Private/LaunchEngineLoop.cpp
@@ -217,9 +217,11 @@
 	#include "ShaderCodeLibrary.h"
 	#include "ShaderPipelineCache.h"
 
-#if !UE_BUILD_SHIPPING
+#if !(UE_BUILD_SHIPPING || UE_BUILD_TEST)
 	#include "ProfileVisualizerModule.h"
+#if UE_DEPRECATED_PROFILER_ENABLED
 	#include "IProfilerServiceModule.h"
+#endif // UE_DEPRECATED_PROFILER_ENABLED
 #endif
 
 #if WITH_AUTOMATION_WORKER
@@ -4344,11 +4346,13 @@ int32 FEngineLoop::PreInitPostStartupScreen(const TCHAR* CmdLine)
 	// Initialize profile visualizers.
 #if !(UE_BUILD_SHIPPING || UE_BUILD_TEST)
 	FModuleManager::Get().LoadModule(TEXT("ProfileVisualizer"));
+#if UE_DEPRECATED_PROFILER_ENABLED
 	if (FPlatformProcess::SupportsMultithreading())
 	{
 		FModuleManager::Get().LoadModule(TEXT("ProfilerService"));
 		FModuleManager::Get().GetModuleChecked<IProfilerServiceModule>("ProfilerService").CreateProfilerServiceManager();
 	}
+#endif // UE_DEPRECATED_PROFILER_ENABLED
 #endif
 
 	// Init HighRes screenshot system, unless running on server
