Fix crashes when the movie player is disabled
 - GetMoviePlayer may return nullptr if the player was never initialized. Callers should first check IsMoviePlayerEnabled()

diff --git a/Engine/Source/Runtime/Launch/Private/LaunchEngineLoop.cpp b/Engine/Source/Runtime/Launch/Private/LaunchEngineLoop.cpp
--- a/Engine/Source/Runtime/Launch/Private/LaunchEngineLoop.cpp
+++ b/Engine/Source/Runtime/Launch/Private/LaunchEngineLoop.cpp
@@ -3582,6 +3582,8 @@ int32 FEngineLoop::PreInitPostStartupScreen(const TCHAR* CmdLine)
 
 			TSharedPtr<FSlateRenderer> SlateRenderer = PreInitContext.SlateRenderer;
 			TSharedRef<FSlateRenderer> SlateRendererSharedRef = SlateRenderer.ToSharedRef();
+			
+			if (IsMoviePlayerEnabled())
 			{
 				SCOPED_BOOT_TIMING("GetMoviePlayer()->SetupLoadingScreenFromIni");
 				// allow the movie player to load a sequence from the .inis (a PreLoadingScreen module could have already initialized a sequence, in which case
@@ -3605,7 +3607,7 @@ int32 FEngineLoop::PreInitPostStartupScreen(const TCHAR* CmdLine)
 			}
 			// If not using the bundle manager, config will be reloaded after ESP, see below
 
-			if (GetMoviePlayer()->HasEarlyStartupMovie())
+			if (IsMoviePlayerEnabled() && GetMoviePlayer()->HasEarlyStartupMovie())
 			{
 				SCOPED_BOOT_TIMING("EarlyStartupMovie");
 				GetMoviePlayer()->Initialize(SlateRendererSharedRef.Get(), FPreLoadScreenManager::Get() ? FPreLoadScreenManager::Get()->GetRenderWindow() : nullptr);
@@ -3852,7 +3854,7 @@ int32 FEngineLoop::PreInitPostStartupScreen(const TCHAR* CmdLine)
 
 #if !UE_SERVER
     //See if we have an engine loading PreLoadScreen registered, if not try to play an engine loading movie as a backup.
-    if (!IsRunningDedicatedServer() && !IsRunningCommandlet() && !GetMoviePlayer()->IsMovieCurrentlyPlaying())
+    if (IsMoviePlayerEnabled() && !GetMoviePlayer()->IsMovieCurrentlyPlaying())
     {
 		SCOPED_BOOT_TIMING("FPreLoadScreenManager::Get()->Initialize etc");
 		if (FSlateRenderer* Renderer = FSlateApplication::Get().GetRenderer())
@@ -3895,7 +3897,7 @@ int32 FEngineLoop::PreInitPostStartupScreen(const TCHAR* CmdLine)
 
 	// Playing a movie can only happen after the rendering thread is started.
 #if !UE_SERVER// && !UE_EDITOR
-	if (!IsRunningDedicatedServer() && !IsRunningCommandlet() && !GetMoviePlayer()->IsMovieCurrentlyPlaying())
+	if (IsMoviePlayerEnabled() && !GetMoviePlayer()->IsMovieCurrentlyPlaying())
 	{
 		SCOPED_BOOT_TIMING("PlayFirstPreLoadScreen etc");
 		if (FPreLoadScreenManager::Get() && FPreLoadScreenManager::Get()->HasRegisteredPreLoadScreenType(EPreLoadScreenTypes::EngineLoadingScreen))
@@ -3910,13 +3912,14 @@ int32 FEngineLoop::PreInitPostStartupScreen(const TCHAR* CmdLine)
         }
 	}
 #endif
+
 	{
 		SCOPED_BOOT_TIMING("PlatformHandleSplashScreen etc");
 #if !UE_SERVER
 		if (!IsRunningDedicatedServer())
 		{
 			// show or hide splash screen based on movie
-			FPlatformMisc::PlatformHandleSplashScreen(!GetMoviePlayer()->IsMovieCurrentlyPlaying());
+			FPlatformMisc::PlatformHandleSplashScreen(!(IsMoviePlayerEnabled() && GetMoviePlayer()->IsMovieCurrentlyPlaying()));
 		}
 		else
 #endif
@@ -4799,7 +4802,10 @@ int32 FEngineLoop::Init()
 
 	check( GEngine );
 
-	GetMoviePlayer()->PassLoadingScreenWindowBackToGame();
+	if (IsMoviePlayerEnabled())
+	{
+		GetMoviePlayer()->PassLoadingScreenWindowBackToGame();
+	}
     
     if (FPreLoadScreenManager::Get())
     {
@@ -4877,7 +4883,7 @@ int32 FEngineLoop::Init()
 		FPreLoadScreenManager::Get()->SetEngineLoadingComplete(true);
         FPreLoadScreenManager::Get()->WaitForEngineLoadingScreenToFinish();
     }
-    else
+    else if (IsMoviePlayerEnabled())
     {
 		SCOPED_BOOT_TIMING("WaitForMovieToFinish");
 		GetMoviePlayer()->WaitForMovieToFinish();
@@ -5608,7 +5614,7 @@ void FEngineLoop::Tick()
 
 	// Ensure we aren't starting a frame while loading or playing a loading movie
 	FMoviePlayerProxy::BlockingForceFinished();
-	ensure(GetMoviePlayer()->IsLoadingFinished() && !GetMoviePlayer()->IsMovieCurrentlyPlaying());
+	ensure(!IsMoviePlayerEnabled() || (GetMoviePlayer()->IsLoadingFinished() && !GetMoviePlayer()->IsMovieCurrentlyPlaying()));
 
     // Frame profiling kickoff
 #if UE_EXTERNAL_PROFILING_ENABLED
@@ -5870,7 +5876,7 @@ void FEngineLoop::Tick()
 
 #if !UE_SERVER
 				// Is it ok to start up the movie player?
-				if (!IsRunningDedicatedServer() && !IsRunningCommandlet() && !GetMoviePlayer()->IsMovieCurrentlyPlaying())
+				if (IsMoviePlayerEnabled() && !GetMoviePlayer()->IsMovieCurrentlyPlaying())
 				{
 					// Enable the MoviePlayer now that the preload screen manager is done.
 					if (FSlateRenderer* Renderer = FSlateApplication::Get().GetRenderer())
@@ -5883,7 +5889,7 @@ void FEngineLoop::Tick()
                 //Destroy / Clean Up PreLoadScreenManager as we are now done
                 FPreLoadScreenManager::Destroy();
             }
-			else
+			else if (IsMoviePlayerEnabled())
 			{
 				QUICK_SCOPE_CYCLE_COUNTER(STAT_FEngineLoop_WaitForMovieToFinish);
 				GetMoviePlayer()->WaitForMovieToFinish(true);
diff --git a/Engine/Source/Runtime/StreamingPauseRendering/Private/StreamingPauseRendering.cpp b/Engine/Source/Runtime/StreamingPauseRendering/Private/StreamingPauseRendering.cpp
--- a/Engine/Source/Runtime/StreamingPauseRendering/Private/StreamingPauseRendering.cpp
+++ b/Engine/Source/Runtime/StreamingPauseRendering/Private/StreamingPauseRendering.cpp
@@ -96,7 +96,7 @@ void FStreamingPauseRenderingModule::ShutdownModule()
 void FStreamingPauseRenderingModule::BeginStreamingPause( FViewport* GameViewport )
 {
 	// If a movie is already playing don't bother starting one
-	if(GetMoviePlayer()->IsInitialized() && !GetMoviePlayer()->IsMovieCurrentlyPlaying() && IsMoviePlayerEnabled())
+	if (IsMoviePlayerEnabled() && (GetMoviePlayer()->IsInitialized() && !GetMoviePlayer()->IsMovieCurrentlyPlaying()))
 	{
 		TRACE_CPUPROFILER_EVENT_SCOPE(FStreamingPauseRenderingModule::BeginStreamingPause);
 
