Manually load "WorldBookmark" module at engine init

diff --git a/Engine/Source/Runtime/Launch/Private/LaunchEngineLoop.cpp b/Engine/Source/Runtime/Launch/Private/LaunchEngineLoop.cpp
--- a/Engine/Source/Runtime/Launch/Private/LaunchEngineLoop.cpp
+++ b/Engine/Source/Runtime/Launch/Private/LaunchEngineLoop.cpp
@@ -4665,6 +4665,8 @@ bool FEngineLoop::LoadStartupCoreModules()
 	// See UAnimBlueprint::UAnimBlueprint().
 	FModuleManager::Get().LoadModule(TEXT("AnimGraph"));
 
+	FModuleManager::Get().LoadModule(TEXT("WorldBookmark"));
+
 	FModuleManager::Get().LoadModule(TEXT("WorldPartitionEditor"));
 #endif
 
