[CSV] Add metrics for MB loaded across boot and map loading screens.
Follow up: This uses DiskUtilizationTracker which is currently only enabled on Switch. Enabling this on other relevant platforms will be discussed in a follow up.
#tests Tested that the metric is logged as expected on exe override Switch build, and exe override RRs.
#rnx

diff --git a/Engine/Source/Runtime/Launch/Private/LaunchEngineLoop.cpp b/Engine/Source/Runtime/Launch/Private/LaunchEngineLoop.cpp
--- a/Engine/Source/Runtime/Launch/Private/LaunchEngineLoop.cpp
+++ b/Engine/Source/Runtime/Launch/Private/LaunchEngineLoop.cpp
@@ -12,6 +12,7 @@
 #include "Misc/QueuedThreadPoolWrapper.h"
 #include "HAL/FileManager.h"
 #include "HAL/PlatformAffinity.h"
+#include "HAL/DiskUtilizationTracker.h"
 #include "Misc/FileHelper.h"
 #include "Internationalization/TextLocalizationManagerGlobals.h"
 #include "Logging/LogSuppressionInterface.h"
@@ -4760,7 +4761,13 @@ void GameLoopIsStarved()
 
 int32 FEngineLoop::Init()
 {
-	ON_SCOPE_EXIT{ GEngineInitEndTime = FPlatformTime::Seconds(); };
+	ON_SCOPE_EXIT
+	{ 
+#if TRACK_DISK_UTILIZATION
+		CSV_METADATA(TEXT("BootMBLoaded"), *FString::FromInt(GDiskUtilizationTracker.GetLongTermStats().TotalBytesRead / (1024 * 1024)));
+#endif
+		GEngineInitEndTime = FPlatformTime::Seconds(); 
+	};
 	LLM_SCOPE(ELLMTag::EngineInitMemory);
 	SCOPED_BOOT_TIMING("FEngineLoop::Init");
 
