Introduce an alternative way to process suspend game tread for the Android platform.

While investigating Android app suspend workflow found that game thread missed APP_EVENT_STATE_APP_SUSPENDED and continued GEngineLoop.Tick(); and locked on FrameEndSync.Sync call because render thread is already suspended, but should be blocked EventHandlerEvent->Wait();
As a result after resuming render tread generates a frame from the previous cycle.

#rnx

diff --git a/Engine/Source/Programs/UnrealBuildTool/Platform/Android/UEBuildAndroid.cs b/Engine/Source/Programs/UnrealBuildTool/Platform/Android/UEBuildAndroid.cs
--- a/Engine/Source/Programs/UnrealBuildTool/Platform/Android/UEBuildAndroid.cs
+++ b/Engine/Source/Programs/UnrealBuildTool/Platform/Android/UEBuildAndroid.cs
@@ -519,6 +519,9 @@ namespace UnrealBuildTool
 			{
 				Logger.LogInformation("Compiling with USE_ANDROID_STANDALONE");
 				CompileEnvironment.Definitions.Add("USE_ANDROID_STANDALONE=1");
+
+				Logger.LogInformation("Compiling with USE_ANDROID_ALTERNATIVE_SUSPEND");
+				CompileEnvironment.Definitions.Add("USE_ANDROID_ALTERNATIVE_SUSPEND=1");
 			}
 
 			if (Target.bPGOOptimize || Target.bPGOProfile)
diff --git a/Engine/Source/Runtime/Core/Public/Android/AndroidPlatform.h b/Engine/Source/Runtime/Core/Public/Android/AndroidPlatform.h
--- a/Engine/Source/Runtime/Core/Public/Android/AndroidPlatform.h
+++ b/Engine/Source/Runtime/Core/Public/Android/AndroidPlatform.h
@@ -100,6 +100,9 @@ typedef FAndroidTypes FPlatformTypes;
 #ifndef USE_ANDROID_STANDALONE
 	#define USE_ANDROID_STANDALONE					0
 #endif
+#ifndef USE_ANDROID_ALTERNATIVE_SUSPEND
+	#define USE_ANDROID_ALTERNATIVE_SUSPEND			0
+#endif
 
 
 #if (!USE_ANDROID_STANDALONE && UE_BUILD_DEBUG) || (USE_ANDROID_STANDALONE && !UE_BUILD_SHIPPING)
diff --git a/Engine/Source/Runtime/Launch/Private/Android/LaunchAndroid.cpp b/Engine/Source/Runtime/Launch/Private/Android/LaunchAndroid.cpp
--- a/Engine/Source/Runtime/Launch/Private/Android/LaunchAndroid.cpp
+++ b/Engine/Source/Runtime/Launch/Private/Android/LaunchAndroid.cpp
@@ -1519,6 +1519,11 @@ static void SuspendApp_EventThread()
 		EMDoneTrigger->Trigger();
 	}));
 
+#if USE_ANDROID_ALTERNATIVE_SUSPEND
+		// Suspend the GT.
+		FAppEventManager::GetInstance()->EnqueueAppEvent(APP_EVENT_STATE_APP_SUSPENDED);
+#endif
+
 	uint32 StartCycles = FPlatformTime::Cycles();
 
 	FEmbeddedCommunication::WakeGameThread();
@@ -1529,12 +1534,8 @@ static void SuspendApp_EventThread()
 
 	// wait for a period of time before blocking rendering
 	UE_LOG(LogAndroid, Log, TEXT("SuspendApp_EventThread -> , waiting for event manager to process. tid: %d"), FPlatformTLS::GetCurrentThreadId());
-#if USE_ANDROID_STANDALONE
-	//EMDoneTrigger->Reset();
-	bool bSuccess = EMDoneTrigger->Wait(240);
-#else
+
 	bool bSuccess = EMDoneTrigger->Wait(4000);
-#endif
 
 	float ElapsedTimeInMs_EMDoneTrigger_Wait = FPlatformTime::ToMilliseconds(FPlatformTime::Cycles() - StartCycles);
 	UE_CLOG(!bSuccess, LogAndroid, Log, TEXT("SuspendApp_EventThread -> ERROR: backgrounding callback, not responded in timely manner. EMDoneTrigger->Wait, waited '%f' ms"), (float)ElapsedTimeInMs_EMDoneTrigger_Wait);
@@ -1543,7 +1544,10 @@ static void SuspendApp_EventThread()
 	BlockRendering();
 
 	// Suspend the GT.
+#if !USE_ANDROID_ALTERNATIVE_SUSPEND
 	FAppEventManager::GetInstance()->EnqueueAppEvent(APP_EVENT_STATE_APP_SUSPENDED);
+#endif
+
 	UE_LOG(LogAndroid, Log, TEXT("SuspendApp_EventThread(EOF)"));
 }
 
diff --git a/Engine/Source/Runtime/OpenGLDrv/Private/Android/AndroidEGL.cpp b/Engine/Source/Runtime/OpenGLDrv/Private/Android/AndroidEGL.cpp
--- a/Engine/Source/Runtime/OpenGLDrv/Private/Android/AndroidEGL.cpp
+++ b/Engine/Source/Runtime/OpenGLDrv/Private/Android/AndroidEGL.cpp
@@ -1342,6 +1342,7 @@ void BlockRendering()
 		FPlatformProcess::ReturnSynchEventToPool(EventToDelete);
 	});
 
+#if !USE_ANDROID_ALTERNATIVE_SUSPEND
 	// Flush GT first in case it has any dependency on RT work to complete
 	FGraphEventRef GTBlockTask = FFunctionGraphTask::CreateAndDispatchWhenReady([BlockedTrigger]()
 		{
@@ -1350,6 +1351,7 @@ void BlockRendering()
 
 	UE_LOG(LogAndroid, Log, TEXT("Waiting for game thread to release EGL context/surface."));
 	BlockedTrigger->Wait();
+#endif
 
 	// Wait for GC to complete and prevent further GCs
 	FGCScopeGuard GCGuard;
