[Engine]Added new function to get Native InputField's Height. Moved the OnVirtualKeyboardHidden event broadcast allowing to be received even when the app is backgrounded.

diff --git a/Engine/Build/Android/Java/src/com/epicgames/unreal/GameActivity.java.template b/Engine/Build/Android/Java/src/com/epicgames/unreal/GameActivity.java.template
--- a/Engine/Build/Android/Java/src/com/epicgames/unreal/GameActivity.java.template
+++ b/Engine/Build/Android/Java/src/com/epicgames/unreal/GameActivity.java.template
@@ -6111,6 +6111,10 @@ public class GameActivity extends $${gameActivitySuperClass}$$ implements Surfac
 			// make sure it is sane.. some drivers returned crazy values
 			return rate < 15.0f ? 60.0f : rate;
 		}
+		else if (key.equals("unreal.input.virtualKeyboardInputHeight"))
+        {
+        	return newVirtualKeyboardInput != null ? newVirtualKeyboardInput.getHeight() : 0.0f;
+        }
 		return 0.0f;
 	}
 
diff --git a/Engine/Source/Runtime/Core/Private/Android/AndroidPlatformMisc.cpp b/Engine/Source/Runtime/Core/Private/Android/AndroidPlatformMisc.cpp
--- a/Engine/Source/Runtime/Core/Private/Android/AndroidPlatformMisc.cpp
+++ b/Engine/Source/Runtime/Core/Private/Android/AndroidPlatformMisc.cpp
@@ -3304,3 +3304,13 @@ bool FAndroidMisc::RemoveNetworkListener(FDelegateHandle Handle)
 	return bSuccess;
 }
 
+float FAndroidMisc::GetVirtualKeyboardInputHeight()
+{
+#if USE_ANDROID_JNI
+	extern float AndroidThunkCpp_GetMetaDataFloat(const FString& Key);
+	return AndroidThunkCpp_GetMetaDataFloat(TEXT("unreal.input.virtualKeyboardInputHeight"));
+#else
+	return 0.0f;
+#endif	
+}
+
diff --git a/Engine/Source/Runtime/Core/Public/Android/AndroidPlatformMisc.h b/Engine/Source/Runtime/Core/Public/Android/AndroidPlatformMisc.h
--- a/Engine/Source/Runtime/Core/Public/Android/AndroidPlatformMisc.h
+++ b/Engine/Source/Runtime/Core/Public/Android/AndroidPlatformMisc.h
@@ -328,6 +328,8 @@ public:
 	static CORE_API FDelegateHandle AddNetworkListener(FOnNetworkConnectionChangedDelegate&& InNewDelegate);
 	static CORE_API bool RemoveNetworkListener(FDelegateHandle Handle);
 
+	static CORE_API float GetVirtualKeyboardInputHeight();
+
 private:
 	static CORE_API const ANSICHAR* CodeToString(int Signal, int si_code);
 	static CORE_API EDeviceScreenOrientation DeviceOrientation;
diff --git a/Engine/Source/Runtime/Core/Public/GenericPlatform/GenericPlatformMisc.h b/Engine/Source/Runtime/Core/Public/GenericPlatform/GenericPlatformMisc.h
--- a/Engine/Source/Runtime/Core/Public/GenericPlatform/GenericPlatformMisc.h
+++ b/Engine/Source/Runtime/Core/Public/GenericPlatform/GenericPlatformMisc.h
@@ -1955,6 +1955,11 @@ public:
 		return false;
 	}
 
+	static const float GetVirtualKeyboardInputHeight()
+	{
+		return 0.0f;
+	}
+	
 #if !UE_BUILD_SHIPPING
 	/**
 	 * Returns any platform specific warning messages we want printed on screen
diff --git a/Engine/Source/Runtime/Launch/Private/Android/AndroidJNI.cpp b/Engine/Source/Runtime/Launch/Private/Android/AndroidJNI.cpp
--- a/Engine/Source/Runtime/Launch/Private/Android/AndroidJNI.cpp
+++ b/Engine/Source/Runtime/Launch/Private/Android/AndroidJNI.cpp
@@ -789,6 +789,14 @@ JNI_METHOD void Java_com_epicgames_unreal_GameActivity_nativeVirtualKeyboardVisi
 	if (!bShown)
 	{
 		VirtualKeyboardWidget.Reset();
+
+		if (FTaskGraphInterface::IsRunning())
+		{
+			FGraphEventRef VirtualKeyboardShown = FFunctionGraphTask::CreateAndDispatchWhenReady([&]()
+			{
+				FAndroidApplication::Get()->OnVirtualKeyboardHidden().Broadcast();
+			}, TStatId(), NULL, ENamedThreads::GameThread);
+		}
 	}
 }
 
@@ -1077,14 +1085,6 @@ void AndroidThunkCpp_HideVirtualKeyboardInputDialog()
 
         // call the java side
         FJavaWrapper::CallVoidMethod(Env, FJavaWrapper::GameActivityThis, FJavaWrapper::AndroidThunkJava_HideVirtualKeyboardInputDialog);
-
-		if (FTaskGraphInterface::IsRunning())
-		{
-			FGraphEventRef VirtualKeyboardShown = FFunctionGraphTask::CreateAndDispatchWhenReady([&]()
-			{
-				FAndroidApplication::Get()->OnVirtualKeyboardHidden().Broadcast();
-			}, TStatId(), NULL, ENamedThreads::GameThread);
-		}
 	}
 }
 
@@ -1117,14 +1117,6 @@ void AndroidThunkCpp_HideVirtualKeyboardInput()
 
 		// call the java side
 		FJavaWrapper::CallVoidMethod(Env, FJavaWrapper::GameActivityThis, FJavaWrapper::AndroidThunkJava_HideVirtualKeyboardInput);
-
-		if( FTaskGraphInterface::IsRunning() )
-		{
-			FGraphEventRef VirtualKeyboardShown = FFunctionGraphTask::CreateAndDispatchWhenReady( [&]()
-			{
-				FAndroidApplication::Get()->OnVirtualKeyboardHidden().Broadcast();
-			}, TStatId(), NULL, ENamedThreads::GameThread );
-		}
 	}
 }
 
