- Disabled SwiftUI Main on TVOS until we get it working

diff --git a/Engine/Source/Runtime/Launch/Launch.Build.cs b/Engine/Source/Runtime/Launch/Launch.Build.cs
--- a/Engine/Source/Runtime/Launch/Launch.Build.cs
+++ b/Engine/Source/Runtime/Launch/Launch.Build.cs
@@ -240,7 +240,7 @@ public class Launch : ModuleRules
 			}
 		}
 
-		if (Target.IsInPlatformGroup(UnrealPlatformGroup.Apple) && Target.Type != TargetType.Editor)
+		if (Target.IsInPlatformGroup(UnrealPlatformGroup.Apple) && Target.Type != TargetType.Editor && Target.Platform != UnrealTargetPlatform.TVOS)
 		{
 			SwiftInteropHeader = Path.Combine(ModuleDirectory, "Private/Apple/LaunchCPPToSwift.h");
 		}
