- Disabling swift in Launch while we deal with some horde issues

diff --git a/Engine/Source/Runtime/Launch/Launch.Build.cs b/Engine/Source/Runtime/Launch/Launch.Build.cs
--- a/Engine/Source/Runtime/Launch/Launch.Build.cs
+++ b/Engine/Source/Runtime/Launch/Launch.Build.cs
@@ -240,9 +240,9 @@ public class Launch : ModuleRules
 			}
 		}
 
-		if (Target.IsInPlatformGroup(UnrealPlatformGroup.Apple) && Target.Type != TargetType.Editor && Target.Platform != UnrealTargetPlatform.TVOS)
+		if (Target.IsInPlatformGroup(UnrealPlatformGroup.Apple) && Target.Type != TargetType.Editor)
 		{
-			SwiftInteropHeader = Path.Combine(ModuleDirectory, "Private/Apple/LaunchCPPToSwift.h");
+//			SwiftInteropHeader = Path.Combine(ModuleDirectory, "Private/Apple/LaunchCPPToSwift.h");
 		}
 
 		if ((Target.Platform == UnrealTargetPlatform.Win64) ||
