- Enabled swift on VisionOS with a virtual property to enable it

diff --git a/Engine/Platforms/VisionOS/Source/Runtime/Launch/Launch_VisionOS.Build.cs b/Engine/Platforms/VisionOS/Source/Runtime/Launch/Launch_VisionOS.Build.cs
--- a/Engine/Platforms/VisionOS/Source/Runtime/Launch/Launch_VisionOS.Build.cs
+++ b/Engine/Platforms/VisionOS/Source/Runtime/Launch/Launch_VisionOS.Build.cs
@@ -6,8 +6,9 @@ using System.IO;
 
 public class Launch_VisionOS : Launch
 {
+	protected override bool bAllowSwiftMain => true;
+
 	public Launch_VisionOS(ReadOnlyTargetRules Target) : base(Target)
 	{
-		PrivateDependencyModuleNames.Add("Swift");
 	}
 }
diff --git a/Engine/Source/Runtime/Launch/Launch.Build.cs b/Engine/Source/Runtime/Launch/Launch.Build.cs
--- a/Engine/Source/Runtime/Launch/Launch.Build.cs
+++ b/Engine/Source/Runtime/Launch/Launch.Build.cs
@@ -8,6 +8,8 @@ using EpicGames.Core;
 
 public class Launch : ModuleRules
 {
+	protected virtual bool bAllowSwiftMain => false;
+	
 	public Launch(ReadOnlyTargetRules Target) : base(Target)
 	{
 		PrivateIncludePathModuleNames.AddRange(new string[] {
@@ -240,9 +242,9 @@ public class Launch : ModuleRules
 			}
 		}
 
-		if (Target.IsInPlatformGroup(UnrealPlatformGroup.Apple) && Target.Type != TargetType.Editor)
+		if (bAllowSwiftMain)//Target.IsInPlatformGroup(UnrealPlatformGroup.Apple) && Target.Type != TargetType.Editor)
 		{
-//			SwiftInteropHeader = Path.Combine(ModuleDirectory, "Private/Apple/LaunchCPPToSwift.h");
+			SwiftInteropHeader = Path.Combine(ModuleDirectory, "Private/Apple/LaunchCPPToSwift.h");
 		}
 
 		if ((Target.Platform == UnrealTargetPlatform.Win64) ||
