Android intent Uri handling

diff --git a/Engine/Build/Android/Java/src/com/epicgames/unreal/GameActivity.java.template b/Engine/Build/Android/Java/src/com/epicgames/unreal/GameActivity.java.template
--- a/Engine/Build/Android/Java/src/com/epicgames/unreal/GameActivity.java.template
+++ b/Engine/Build/Android/Java/src/com/epicgames/unreal/GameActivity.java.template
@@ -3992,6 +3992,10 @@ public class GameActivity extends $${gameActivitySuperClass}$$ implements Surfac
 		super.onNewIntent(newIntent);
 		setIntent(newIntent);
 //$${gameActivityOnNewIntentAdditions}$$
+		if (newIntent.getData() != null)
+		{
+			nativeHandleNewIntentData(newIntent.getData().toString());
+		}
 	}
 
 	public void onTrimMemory(int level)
@@ -6340,6 +6344,15 @@ public class GameActivity extends $${gameActivitySuperClass}$$ implements Surfac
 		return _extrasBundle.getString(key);
 	}
 
+	public String AndroidThunkJava_GetIntentDataAsString()
+	{
+		if (getIntent() == null || getIntent().getData() == null)
+		{
+			return null;
+		}
+		return getIntent().getData().toString();
+	}
+
 	public void AndroidThunkJava_SetSustainedPerformanceMode(final boolean bEnable)
 	{
 		if (ANDROID_BUILD_VERSION >= 24)
@@ -7205,6 +7218,8 @@ public class GameActivity extends $${gameActivitySuperClass}$$ implements Surfac
 
 	public native void nativeNetworkChanged();
 
+	public native void nativeHandleNewIntentData(String intentUri);
+
 	static
 	{
 		// We need to decide on enabling memory tracing before loading native code.
diff --git a/Engine/Build/Android/Java/src/com/epicgames/unreal/SplashActivity.java b/Engine/Build/Android/Java/src/com/epicgames/unreal/SplashActivity.java
--- a/Engine/Build/Android/Java/src/com/epicgames/unreal/SplashActivity.java
+++ b/Engine/Build/Android/Java/src/com/epicgames/unreal/SplashActivity.java
@@ -198,6 +198,13 @@ public class SplashActivity extends Activity
 			GameActivityIntent.setAction(intentAction);
 		}
 
+		// pass the data if available
+		Uri intentData = intentFromActivity.getData();
+		if (intentData != null)
+		{
+			GameActivityIntent.setData(intentData);
+		}
+
 		// make a list of ungranted dangerous permissions in manifest required at start of GameActivity and request any we still need
 		ArrayList<String> ungrantedPermissions = getUngrantedPermissions(this, getDangerousPermissions(pm, packageName), permissionsRequiredAtStart);
 		if (ungrantedPermissions.size() > 0)
diff --git a/Engine/Source/Runtime/Core/Private/Android/AndroidPlatformMisc.cpp b/Engine/Source/Runtime/Core/Private/Android/AndroidPlatformMisc.cpp
--- a/Engine/Source/Runtime/Core/Private/Android/AndroidPlatformMisc.cpp
+++ b/Engine/Source/Runtime/Core/Private/Android/AndroidPlatformMisc.cpp
@@ -101,6 +101,14 @@ static FAutoConsoleVariableRef CAndroidMemoryStateChangeThreshold(
 	ECVF_Default
 );
 
+static bool GAndroidBroadcastIntentData = false;
+static FAutoConsoleVariableRef CVarAndroidBroadcastIntentData(
+	TEXT("android.BroadcastIntentData"),
+	GAndroidBroadcastIntentData,
+	TEXT("Whether to broadcast intent data, retry every frame if delegate is not bound"),
+	ECVF_Default
+);
+
 #if STATS || ENABLE_STATNAMEDEVENTS
 int32 FAndroidMisc::TraceMarkerFileDescriptor = -1;
 
@@ -1655,6 +1663,52 @@ bool FAndroidMisc::IsAllowedRemoteNotifications()
 #endif
 }
 
+static void HandleNewIntentUri_GameThread(const FString& IntentUri)
+{
+	static FString ProtocolActivationUri;
+	static FTSTicker::FDelegateHandle ProtocolActivationTickerHandle;
+
+	ProtocolActivationUri = IntentUri;
+
+	if (GAndroidBroadcastIntentData && !ProtocolActivationUri.IsEmpty())
+	{
+		if (!ProtocolActivationTickerHandle.IsValid())
+		{
+			ProtocolActivationTickerHandle = FTSTicker::GetCoreTicker().AddTicker(FTickerDelegate::CreateLambda(
+				[](float)
+				{
+					if (FCoreDelegates::OnActivatedByProtocol.IsBound())
+					{
+						FCoreDelegates::OnActivatedByProtocol.Broadcast(ProtocolActivationUri, PLATFORMUSERID_NONE);
+
+						ProtocolActivationUri.Reset();
+						ProtocolActivationTickerHandle.Reset();
+						return false; // remove from ticker
+					}
+
+					return true; // try again next tick
+				}
+			));
+		}
+	}
+	else if (ProtocolActivationTickerHandle.IsValid())
+	{
+		FTSTicker::GetCoreTicker().RemoveTicker(ProtocolActivationTickerHandle);
+		ProtocolActivationTickerHandle.Reset();
+	}
+}
+
+void FAndroidMisc::HandleNewIntentUri(const FString& IntentUri)
+{
+	if (!IntentUri.IsEmpty())
+	{
+		FFunctionGraphTask::CreateAndDispatchWhenReady([IntentUri]()
+			{
+				HandleNewIntentUri_GameThread(IntentUri);
+			}, TStatId(), NULL, ENamedThreads::GameThread);
+	}
+}
+
 TArray<uint8> FAndroidMisc::GetSystemFontBytes()
 {
 #if USE_ANDROID_FILE
diff --git a/Engine/Source/Runtime/Core/Public/Android/AndroidPlatformMisc.h b/Engine/Source/Runtime/Core/Public/Android/AndroidPlatformMisc.h
--- a/Engine/Source/Runtime/Core/Public/Android/AndroidPlatformMisc.h
+++ b/Engine/Source/Runtime/Core/Public/Android/AndroidPlatformMisc.h
@@ -156,6 +156,8 @@ public:
 	static CORE_API void UnregisterForRemoteNotifications();
 	static CORE_API bool IsAllowedRemoteNotifications();
 
+	static CORE_API void HandleNewIntentUri(const FString& IntentUri);
+
 	/** @return Memory representing a true type or open type font provided by the platform as a default font for unreal to consume; empty array if the default font failed to load. */
 	static CORE_API TArray<uint8> GetSystemFontBytes();
 
diff --git a/Engine/Source/Runtime/Launch/Private/Android/AndroidJNI.cpp b/Engine/Source/Runtime/Launch/Private/Android/AndroidJNI.cpp
--- a/Engine/Source/Runtime/Launch/Private/Android/AndroidJNI.cpp
+++ b/Engine/Source/Runtime/Launch/Private/Android/AndroidJNI.cpp
@@ -132,6 +132,7 @@ void FJavaWrapper::FindClassesAndMethods(JNIEnv* Env)
 	AndroidThunkJava_GetIntentExtrasBoolean = FindMethod(Env, GameActivityClassID, "AndroidThunkJava_GetIntentExtrasBoolean", "(Ljava/lang/String;)Z", bIsOptional);
 	AndroidThunkJava_GetIntentExtrasInt = FindMethod(Env, GameActivityClassID, "AndroidThunkJava_GetIntentExtrasInt", "(Ljava/lang/String;)I", bIsOptional);
 	AndroidThunkJava_GetIntentExtrasString = FindMethod(Env, GameActivityClassID, "AndroidThunkJava_GetIntentExtrasString", "(Ljava/lang/String;)Ljava/lang/String;", bIsOptional);
+	AndroidThunkJava_GetIntentDataAsString = FindMethod(Env, GameActivityClassID, "AndroidThunkJava_GetIntentDataAsString", "()Ljava/lang/String;", bIsOptional);
 	AndroidThunkJava_PushSensorEvents = FindMethod(Env, GameActivityClassID, "AndroidThunkJava_PushSensorEvents", "()V", bIsOptional);
 	AndroidThunkJava_SetOrientation = FindMethod(Env, GameActivityClassID, "AndroidThunkJava_SetOrientation", "(I)V", bIsOptional);
 	AndroidThunkJava_HasSharedPreference = FindMethod(Env, GameActivityClassID, "AndroidThunkJava_HasSharedPreference", "(Ljava/lang/String;Ljava/lang/String;)Z", bIsOptional);
@@ -542,6 +543,7 @@ jmethodID FJavaWrapper::AndroidThunkJava_HasIntentExtrasKey;
 jmethodID FJavaWrapper::AndroidThunkJava_GetIntentExtrasBoolean;
 jmethodID FJavaWrapper::AndroidThunkJava_GetIntentExtrasInt;
 jmethodID FJavaWrapper::AndroidThunkJava_GetIntentExtrasString;
+jmethodID FJavaWrapper::AndroidThunkJava_GetIntentDataAsString;
 jmethodID FJavaWrapper::AndroidThunkJava_SetSustainedPerformanceMode;
 jmethodID FJavaWrapper::AndroidThunkJava_PushSensorEvents;
 jmethodID FJavaWrapper::AndroidThunkJava_IsScreenCaptureDisabled;
@@ -770,6 +772,11 @@ void AndroidThunkCpp_RestartApplication(const FString& IntentString)
 	}
 }
 
+JNI_METHOD void Java_com_epicgames_unreal_GameActivity_nativeHandleNewIntentData(JNIEnv* jenv, jobject thiz, jstring JavaIntent)
+{
+	FAndroidMisc::HandleNewIntentUri(FJavaHelper::FStringFromParam(jenv, JavaIntent));
+}
+
 JNI_METHOD void Java_com_epicgames_unreal_NativeCalls_WebViewVisible(JNIEnv* jenv, jobject thiz, jboolean bShown)
 {
 	GWebViewShown = bShown;
@@ -925,6 +932,16 @@ FString AndroidThunkCpp_GetIntentExtrasString(const FString& Key)
 	return Result;
 }
 
+FString AndroidThunkCpp_GetIntentDataAsString()
+{
+	FString Result = FString("");
+	if (JNIEnv* Env = FAndroidApplication::GetJavaEnv())
+	{
+		Result = FJavaHelper::FStringFromLocalRef(Env, (jstring)FJavaWrapper::CallObjectMethod(Env, FJavaWrapper::GameActivityThis, FJavaWrapper::AndroidThunkJava_GetIntentDataAsString));
+	}
+	return Result;
+}
+
 void AndroidThunkCpp_SetSustainedPerformanceMode(bool bEnable)
 {
 	if (JNIEnv* Env = FAndroidApplication::GetJavaEnv())
diff --git a/Engine/Source/Runtime/Launch/Private/Android/LaunchAndroid.cpp b/Engine/Source/Runtime/Launch/Private/Android/LaunchAndroid.cpp
--- a/Engine/Source/Runtime/Launch/Private/Android/LaunchAndroid.cpp
+++ b/Engine/Source/Runtime/Launch/Private/Android/LaunchAndroid.cpp
@@ -179,6 +179,7 @@ extern bool AndroidThunkCpp_IsVirtualKeyboardShown();
 extern bool AndroidThunkCpp_IsWebViewShown();
 extern void AndroidThunkCpp_RestartApplication(const FString& IntentString);
 extern FString AndroidThunkCpp_GetIntentExtrasString(const FString& Key);
+extern FString AndroidThunkCpp_GetIntentDataAsString();
 
 // Base path for file accesses
 extern FString GFilePathBase;
@@ -732,6 +733,8 @@ int32 AndroidMain(struct android_app* state)
 
 	FAndroidStats::Init(FParse::Param(FCommandLine::Get(), TEXT("hwcpipe")));
 
+	FAndroidMisc::HandleNewIntentUri(AndroidThunkCpp_GetIntentDataAsString());
+
 	BootTimingPoint("Tick loop starting");
 	DumpBootTiming();
 
diff --git a/Engine/Source/Runtime/Launch/Public/Android/AndroidJNI.h b/Engine/Source/Runtime/Launch/Public/Android/AndroidJNI.h
--- a/Engine/Source/Runtime/Launch/Public/Android/AndroidJNI.h
+++ b/Engine/Source/Runtime/Launch/Public/Android/AndroidJNI.h
@@ -70,6 +70,7 @@ public:
 	static jmethodID AndroidThunkJava_GetIntentExtrasBoolean;
 	static jmethodID AndroidThunkJava_GetIntentExtrasInt;
 	static jmethodID AndroidThunkJava_GetIntentExtrasString;
+	static jmethodID AndroidThunkJava_GetIntentDataAsString;
 	static jmethodID AndroidThunkJava_SetSustainedPerformanceMode;
 	static jmethodID AndroidThunkJava_PushSensorEvents;
 	static jmethodID AndroidThunkJava_SetOrientation;
