fix for UnrealLightmass binary fails to launch

#rnx

diff --git a/Engine/Source/Runtime/Launch/Private/LaunchEngineLoop.cpp b/Engine/Source/Runtime/Launch/Private/LaunchEngineLoop.cpp
--- a/Engine/Source/Runtime/Launch/Private/LaunchEngineLoop.cpp
+++ b/Engine/Source/Runtime/Launch/Private/LaunchEngineLoop.cpp
@@ -2189,8 +2189,16 @@ int32 FEngineLoop::PreInitPreStartupScreen(const TCHAR* CmdLine)
 			FString::Printf(TEXT("%s-%sPreInit-%s-%s.%s"), *FApp::GetName(), FApp::GetProjectName(), FPlatformProcess::GetBinariesSubdirectory(), LexToString(FApp::GetBuildConfiguration()), FPlatformProcess::GetModuleExtension())
 #endif
 		};
-
+		
+		// UE-230629: On mac UnrealLightmass this would cause assert on main thread
+#if PLATFORM_MAC
+		if (![NSThread isMainThread])
+		{
+#endif
 		FTaskTagScope scope(ETaskTag::EStaticInit);
+#if PLATFORM_MAC
+		}
+#endif
 
 		for (const FString& FileName : FileNames)
 		{
