FEditorDelegates : Add an OnEditorPreExit delegate.

This delegate provides a means to receive notification before the editor subsystems are deinitialized.

diff --git a/Engine/Source/Editor/UnrealEd/Private/Editor.cpp b/Engine/Source/Editor/UnrealEd/Private/Editor.cpp
--- a/Engine/Source/Editor/UnrealEd/Private/Editor.cpp
+++ b/Engine/Source/Editor/UnrealEd/Private/Editor.cpp
@@ -179,6 +179,7 @@ FEditorDelegates::FOnViewAssetIdentifiers				FEditorDelegates::OnEditAssetIdenti
 FEditorDelegates::FOnRestartRequested					FEditorDelegates::OnRestartRequested;
 FEditorDelegates::FOnEditorBoot							FEditorDelegates::OnEditorBoot;
 FEditorDelegates::FOnEditorInitialized					FEditorDelegates::OnEditorInitialized;
+FSimpleMulticastDelegate								FEditorDelegates::OnEditorPreExit;
 FEditorDelegates::FOnExternalContentResolved			FEditorDelegates::OnExternalContentResolved;
 
 PRAGMA_ENABLE_DEPRECATION_WARNINGS
diff --git a/Engine/Source/Editor/UnrealEd/Public/Editor.h b/Engine/Source/Editor/UnrealEd/Public/Editor.h
--- a/Engine/Source/Editor/UnrealEd/Public/Editor.h
+++ b/Engine/Source/Editor/UnrealEd/Public/Editor.h
@@ -380,6 +380,8 @@ struct FEditorDelegates
 	static UNREALED_API FOnEditorBoot OnEditorBoot;
 	/** Called when the editor has initialized */
 	static UNREALED_API FOnEditorInitialized OnEditorInitialized;
+	/** Called before the editor exits. This is called before any editor subsystems are deinitialized, and also before FCoreDelegates::OnEnginePreExit */
+	static UNREALED_API FSimpleMulticastDelegate OnEditorPreExit;
 	/** Called when external content gets resolved */
 	static UNREALED_API FOnExternalContentResolved OnExternalContentResolved;
 };
diff --git a/Engine/Source/Runtime/Launch/Private/LaunchEngineLoop.cpp b/Engine/Source/Runtime/Launch/Private/LaunchEngineLoop.cpp
--- a/Engine/Source/Runtime/Launch/Private/LaunchEngineLoop.cpp
+++ b/Engine/Source/Runtime/Launch/Private/LaunchEngineLoop.cpp
@@ -4320,6 +4320,13 @@ int32 FEngineLoop::PreInitPostStartupScreen(const TCHAR* CmdLine)
 	{
 		if ( GEngine != nullptr )
 		{
+#if WITH_EDITOR
+			if (GIsEditor)
+			{
+				FEditorDelegates::OnEditorPreExit.Broadcast();
+			}
+#endif // WITH_EDITOR
+
 			GEngine->PreExit();
 		}
 		AppPreExit();
@@ -5067,6 +5074,13 @@ void FEngineLoop::Exit()
 
 	if ( GEngine != nullptr )
 	{
+#if WITH_EDITOR
+		if (GIsEditor)
+		{
+			FEditorDelegates::OnEditorPreExit.Broadcast();
+		}
+#endif // WITH_EDITOR
+
 		GEngine->PreExit();
 	}
 
