Load the MaterialEditor module at startup so we register editor preference settings on startup.

- Before this change the "Material Editor" preference settings wouldn't show up unless you opened the material editor first.
- This was because we register the material editor settings when the Material Editor module is first loaded.
- Most Editors default load their module during startup so it makes sense to have the material editor do the same.
- We _shouldn't_ need to load this in the dedicate server.

[RN] minor Rendering - Materials and Shaders

diff --git a/Engine/Source/Runtime/Launch/Private/LaunchEngineLoop.cpp b/Engine/Source/Runtime/Launch/Private/LaunchEngineLoop.cpp
--- a/Engine/Source/Runtime/Launch/Private/LaunchEngineLoop.cpp
+++ b/Engine/Source/Runtime/Launch/Private/LaunchEngineLoop.cpp
@@ -4662,6 +4662,13 @@ bool FEngineLoop::LoadStartupCoreModules()
 
 #endif //(WITH_EDITOR && !(UE_BUILD_SHIPPING || UE_BUILD_TEST))
 
+#if WITH_EDITOR
+	if (!IsRunningDedicatedServer())
+	{
+		FModuleManager::Get().LoadModule(TEXT("MaterialEditor"));
+	}
+#endif // WITH_EDITOR
+
 #if WITH_ENGINE
 	// Load runtime client modules (which are also needed at cook-time)
 	if( !IsRunningDedicatedServer() )
