[ASIS] Fix build error. Use Add() method to add element to the array

#rnx

diff --git a/Engine/Source/Runtime/Launch/Private/Android/LaunchAndroid.cpp b/Engine/Source/Runtime/Launch/Private/Android/LaunchAndroid.cpp
--- a/Engine/Source/Runtime/Launch/Private/Android/LaunchAndroid.cpp
+++ b/Engine/Source/Runtime/Launch/Private/Android/LaunchAndroid.cpp
@@ -274,7 +274,7 @@ static void IssueConsoleCommand(FString Command)
 	else
 	{
 		STANDALONE_DEBUG_LOGf(LogAndroid, TEXT("Ignoring console command (too early, Queue for later): %s"), *Command);
-		GPendingConsoleCommands.Add_GetRef(MoveTemp(Command));
+		GPendingConsoleCommands.Add(MoveTemp(Command));
 	}
 #endif
 }
