Delete RealtimeProfiler module
- The module and its contents have been unused and unmaintained for a long time.

diff --git a/Engine/Source/Developer/RealtimeProfiler/Private/RealtimeProfiler.cpp b/Engine/Source/Developer/RealtimeProfiler/Private/RealtimeProfiler.cpp
deleted file mode 100644
diff --git a/Engine/Source/Developer/RealtimeProfiler/Private/RealtimeProfiler.h b/Engine/Source/Developer/RealtimeProfiler/Private/RealtimeProfiler.h
deleted file mode 100644
diff --git a/Engine/Source/Developer/RealtimeProfiler/Private/SRealtimeProfilerFrame.cpp b/Engine/Source/Developer/RealtimeProfiler/Private/SRealtimeProfilerFrame.cpp
deleted file mode 100644
diff --git a/Engine/Source/Developer/RealtimeProfiler/Private/SRealtimeProfilerFrame.h b/Engine/Source/Developer/RealtimeProfiler/Private/SRealtimeProfilerFrame.h
deleted file mode 100644
diff --git a/Engine/Source/Developer/RealtimeProfiler/Private/SRealtimeProfilerLineGraph.cpp b/Engine/Source/Developer/RealtimeProfiler/Private/SRealtimeProfilerLineGraph.cpp
deleted file mode 100644
diff --git a/Engine/Source/Developer/RealtimeProfiler/Private/SRealtimeProfilerLineGraph.h b/Engine/Source/Developer/RealtimeProfiler/Private/SRealtimeProfilerLineGraph.h
deleted file mode 100644
diff --git a/Engine/Source/Developer/RealtimeProfiler/Private/SRealtimeProfilerTimeline.cpp b/Engine/Source/Developer/RealtimeProfiler/Private/SRealtimeProfilerTimeline.cpp
deleted file mode 100644
diff --git a/Engine/Source/Developer/RealtimeProfiler/Private/SRealtimeProfilerTimeline.h b/Engine/Source/Developer/RealtimeProfiler/Private/SRealtimeProfilerTimeline.h
deleted file mode 100644
diff --git a/Engine/Source/Developer/RealtimeProfiler/Private/SRealtimeProfilerVisualizer.cpp b/Engine/Source/Developer/RealtimeProfiler/Private/SRealtimeProfilerVisualizer.cpp
deleted file mode 100644
diff --git a/Engine/Source/Developer/RealtimeProfiler/Private/SRealtimeProfilerVisualizer.h b/Engine/Source/Developer/RealtimeProfiler/Private/SRealtimeProfilerVisualizer.h
deleted file mode 100644
diff --git a/Engine/Source/Developer/RealtimeProfiler/RealtimeProfiler.Build.cs b/Engine/Source/Developer/RealtimeProfiler/RealtimeProfiler.Build.cs
deleted file mode 100644
diff --git a/Engine/Source/Runtime/Launch/Launch.Build.cs b/Engine/Source/Runtime/Launch/Launch.Build.cs
--- a/Engine/Source/Runtime/Launch/Launch.Build.cs
+++ b/Engine/Source/Runtime/Launch/Launch.Build.cs
@@ -173,7 +173,6 @@ public class Launch : ModuleRules
 		{
 			DynamicallyLoadedModuleNames.AddRange(new string[] {
 				"ProfileVisualizer",
-				"RealtimeProfiler",
 			});
 
 			if (Target.GlobalDefinitions.Contains("UE_DEPRECATED_PROFILER_ENABLED=1"))
