- Fixed #include issue with LiveLinkHub

diff --git a/Engine/Source/Runtime/Core/Internal/CoreGlobalsInternal.h b/Engine/Source/Runtime/Core/Internal/CoreGlobalsInternal.h
--- a/Engine/Source/Runtime/Core/Internal/CoreGlobalsInternal.h
+++ b/Engine/Source/Runtime/Core/Internal/CoreGlobalsInternal.h
@@ -3,9 +3,12 @@
 
 #include "CoreTypes.h"
 #include "Misc/Build.h"
+#include "Misc/Parse.h"
+#include "Misc/CommandLine.h"
 
 #if WITH_EDITOR
 extern CORE_API bool PRIVATE_GIsRunningHybridCookedEditor;
+extern CORE_API bool PRIVATE_GHasInitializedHybridCookedEditor;
 #endif
 
 /**
@@ -14,6 +17,11 @@ extern CORE_API bool PRIVATE_GIsRunningHybridCookedEditor;
 FORCEINLINE bool IsRunningHybridCookedEditor()
 {
 #if WITH_EDITOR
+	if (!PRIVATE_GHasInitializedHybridCookedEditor)
+	{
+		PRIVATE_GIsRunningHybridCookedEditor = FParse::Param(FCommandLine::Get(), TEXT("hybridcookededitor"));
+		PRIVATE_GHasInitializedHybridCookedEditor = true;
+	}
 	return PRIVATE_GIsRunningHybridCookedEditor;
 #else
 	return false;
diff --git a/Engine/Source/Runtime/Core/Private/Misc/CoreGlobals.cpp b/Engine/Source/Runtime/Core/Private/Misc/CoreGlobals.cpp
--- a/Engine/Source/Runtime/Core/Private/Misc/CoreGlobals.cpp
+++ b/Engine/Source/Runtime/Core/Private/Misc/CoreGlobals.cpp
@@ -178,6 +178,7 @@ bool GIsRunningUnattendedScript = false;
 bool					PRIVATE_GIsRunningCookCommandlet	= false;				/** Whether this executable is running the cook commandlet */
 bool					PRIVATE_GIsRunningDLCCookCommandlet = false;				/** Whether this executable is running the cook commandlet on a DLC plugin */
 bool					PRIVATE_GIsRunningHybridCookedEditor= false;				/** Whether the editor is running in hybrid cooked editor mode */
+bool					PRIVATE_GHasInitializedHybridCookedEditor = false;			/** Whether PRIVATE_GIsRunningHybridCookedEditor has been initialized */
 namespace UE::Private
 {
 int32					GMultiprocessId = 0;
diff --git a/Engine/Source/Runtime/Launch/Private/LaunchEngineLoop.cpp b/Engine/Source/Runtime/Launch/Private/LaunchEngineLoop.cpp
--- a/Engine/Source/Runtime/Launch/Private/LaunchEngineLoop.cpp
+++ b/Engine/Source/Runtime/Launch/Private/LaunchEngineLoop.cpp
@@ -133,7 +133,6 @@
 	#include "Virtualization/VirtualizationSystem.h"
 	#include "AnimationUtils.h"
 	#include "UObject/ICookInfo.h"
-	#include "CoreGlobalsInternal.h"
 
 	#if PLATFORM_WINDOWS
 		#include "Windows/AllowWindowsPlatformTypes.h"
@@ -1951,13 +1950,6 @@ int32 FEngineLoop::PreInitPreStartupScreen(const TCHAR* CmdLine)
 		GShouldEmitVerboseNamedEvents = true;
 	}
 
-#if WITH_EDITOR
-	// we are running in HybridCookedEditor if it's on the commandline
-	if (FParse::Param(CmdLine, TEXT("hybridcookededitor")))
-	{
-		PRIVATE_GIsRunningHybridCookedEditor = true;
-	}
-#endif
 
 	// Set the flag for whether we've build DebugGame instead of Development. The engine does not know this (whereas the launch module does) because it is always built in development.
 #if UE_BUILD_DEVELOPMENT && defined(UE_BUILD_DEVELOPMENT_WITH_DEBUGGAME) && UE_BUILD_DEVELOPMENT_WITH_DEBUGGAME
