VisionOS | Enable Mixed Immersion mode by default if VisionOS 2+ is used. This allows Mixed Immersion Mode to work in Binary/Launcher-version of the Editor. Note that /Engine/ or /Project/'s VisonOSEngine.ini settings must also be enabled as outlined in the AVP Quick Start guide for Passthrough to work.

diff --git a/Engine/Source/Runtime/Launch/Private/Apple/SwiftMain.swift b/Engine/Source/Runtime/Launch/Private/Apple/SwiftMain.swift
--- a/Engine/Source/Runtime/Launch/Private/Apple/SwiftMain.swift
+++ b/Engine/Source/Runtime/Launch/Private/Apple/SwiftMain.swift
@@ -103,9 +103,9 @@ struct UESwiftApp: App {
 	// VisionOS 1.x only supports full immersion style in Metal rendering apps.
 	@State private var style: ImmersionStyle = .full
 #else
-	// VisonOS 2+ also supported mixed immersion.
-	@State private var style: ImmersionStyle = .full
-	//@State private var style: ImmersionStyle = .mixed
+	// VisonOS 2+ also supports mixed immersion, which we enable by default.
+	//@State private var style: ImmersionStyle = .full
+	@State private var style: ImmersionStyle = .mixed
 #endif
 	
 	@UIApplicationDelegateAdaptor(IOSAppDelegate.self) var delegate
