[Shaders] Add <PROJECT>/Shaders/ folder to virtual shader include directories engine wide.

This allows content-only projects to include shader files in custom HLSL expressions outside the Engine/Shaders/ folder.

#rnx

diff --git a/Engine/Source/Runtime/Launch/Private/LaunchEngineLoop.cpp b/Engine/Source/Runtime/Launch/Private/LaunchEngineLoop.cpp
--- a/Engine/Source/Runtime/Launch/Private/LaunchEngineLoop.cpp
+++ b/Engine/Source/Runtime/Launch/Private/LaunchEngineLoop.cpp
@@ -2716,6 +2716,13 @@ int32 FEngineLoop::PreInitPreStartupScreen(const TCHAR* CmdLine)
 	// Add the default engine shader dir
 	AddShaderSourceDirectoryMapping(TEXT("/Engine"), FPlatformProcess::ShaderDir());
 
+	// Add <PROJECT>/Shaders/ folder to virtual shader include directories
+	const FString ProjectShaderPath = FPaths::Combine(FPaths::ProjectDir(), TEXT("Shaders"));
+	if (FPaths::DirectoryExists(ProjectShaderPath))
+	{
+		AddShaderSourceDirectoryMapping(TEXT("/Project"), ProjectShaderPath);
+	}
+
 #if WITH_EDITOR
 	{
 		IPlatformFile& PlatformFile = FPlatformFileManager::Get().GetPlatformFile();
