Filter out touches on WebView control so they don't reacg GameActivity

diff --git a/Engine/Build/Android/Java/src/com/epicgames/unreal/GameActivity.java.template b/Engine/Build/Android/Java/src/com/epicgames/unreal/GameActivity.java.template
--- a/Engine/Build/Android/Java/src/com/epicgames/unreal/GameActivity.java.template
+++ b/Engine/Build/Android/Java/src/com/epicgames/unreal/GameActivity.java.template
@@ -397,6 +397,8 @@ public class GameActivity extends $${gameActivitySuperClass}$$ implements Surfac
 	// number of times screen capturing has been disabled
 	private int NumTimesScreenCaptureDisabled = 0;
 
+	protected List<View> ExtraIgnoreClickViews = new ArrayList<View>();
+
 	// Virtual keyboard
 	AlertDialog virtualKeyboardAlert;
 	EditText virtualKeyboardInputBox;
@@ -6859,9 +6861,24 @@ public class GameActivity extends $${gameActivitySuperClass}$$ implements Surfac
 
 	}
 
+	public boolean HasExtraIgnoreView()
+	{
+		return !ExtraIgnoreClickViews.isEmpty();
+	}
+
+	public void AddExtraIgnoreView(View view)
+	{
+		ExtraIgnoreClickViews.add(view);
+	}
+
+	public void RemoveExtraIgnoreView(View view)
+	{
+		ExtraIgnoreClickViews.remove(view);
+	}
+
 	protected List<View> GetExtraIgnoreViews()
 	{
-		return null;
+		return ExtraIgnoreClickViews;
 	}
 
 	//check if the new virtual keyboard input has received a MOUSE_DOWN event
diff --git a/Engine/Build/Android/Java/src/com/epicgames/unreal/WebViewControl.java b/Engine/Build/Android/Java/src/com/epicgames/unreal/WebViewControl.java
--- a/Engine/Build/Android/Java/src/com/epicgames/unreal/WebViewControl.java
+++ b/Engine/Build/Android/Java/src/com/epicgames/unreal/WebViewControl.java
@@ -512,6 +512,7 @@ class WebViewControl
 				}
 				if (!bShown)
 				{
+					GameActivity.Get().AddExtraIgnoreView(webView);
 					bShown = true;
 					
 					// add to the activitiy, on top of the SurfaceView
@@ -605,6 +606,7 @@ class WebViewControl
 			{
 				if (bShown)
 				{
+					GameActivity.Get().RemoveExtraIgnoreView(webView);
 					ViewGroup parent = (ViewGroup)webView.getParent();
 					if (parent != null)
 					{
diff --git a/Engine/Source/Runtime/Launch/Private/Android/AndroidJNI.cpp b/Engine/Source/Runtime/Launch/Private/Android/AndroidJNI.cpp
--- a/Engine/Source/Runtime/Launch/Private/Android/AndroidJNI.cpp
+++ b/Engine/Source/Runtime/Launch/Private/Android/AndroidJNI.cpp
@@ -150,6 +150,8 @@ void FJavaWrapper::FindClassesAndMethods(JNIEnv* Env)
 	AndroidThunkJava_IsScreenCaptureDisabled = FindMethod(Env, GameActivityClassID, "AndroidThunkJava_IsScreenCaptureDisabled", "()Z", bIsOptional);
 	AndroidThunkJava_DisableScreenCapture = FindMethod(Env, GameActivityClassID, "AndroidThunkJava_DisableScreenCapture", "(Z)V", bIsOptional);
 
+	// Native view click filtering
+	AndroidThunkJava_HasExtraIgnoreView = FindMethod(Env, GameActivityClassID, "HasExtraIgnoreView", "()Z", bIsOptional);;
 
 	// this is optional - only inserted if Oculus Mobile plugin enabled
 	AndroidThunkJava_IsOculusMobileApplication = FindMethod(Env, GameActivityClassID, "AndroidThunkJava_IsOculusMobileApplication", "()Z", true);
@@ -549,6 +551,7 @@ jmethodID FJavaWrapper::AndroidThunkJava_SetSustainedPerformanceMode;
 jmethodID FJavaWrapper::AndroidThunkJava_PushSensorEvents;
 jmethodID FJavaWrapper::AndroidThunkJava_IsScreenCaptureDisabled;
 jmethodID FJavaWrapper::AndroidThunkJava_DisableScreenCapture;
+jmethodID FJavaWrapper::AndroidThunkJava_HasExtraIgnoreView;
 jmethodID FJavaWrapper::AndroidThunkJava_SetOrientation;
 jmethodID FJavaWrapper::AndroidThunkJava_HasSharedPreference;
 jmethodID FJavaWrapper::AndroidThunkJava_SetSharedPreferenceBoolean;
@@ -765,6 +768,16 @@ bool AndroidThunkCpp_VirtualInputIgnoreClick(int32 x, int32 y)
 	return Result;
 }
 
+bool AndroidThunkCpp_HasExtraIgnoreView()
+{
+	bool Result = false;
+	if (JNIEnv* Env = FAndroidApplication::GetJavaEnv())
+	{
+		Result = FJavaWrapper::CallBooleanMethod(Env, FJavaWrapper::GameActivityThis, FJavaWrapper::AndroidThunkJava_HasExtraIgnoreView);
+	}
+	return Result;
+}
+
 void AndroidThunkCpp_RestartApplication(const FString& IntentString)
 {
 	if (JNIEnv* Env = FAndroidApplication::GetJavaEnv())
diff --git a/Engine/Source/Runtime/Launch/Private/Android/LaunchAndroid.cpp b/Engine/Source/Runtime/Launch/Private/Android/LaunchAndroid.cpp
--- a/Engine/Source/Runtime/Launch/Private/Android/LaunchAndroid.cpp
+++ b/Engine/Source/Runtime/Launch/Private/Android/LaunchAndroid.cpp
@@ -177,6 +177,7 @@ extern void AndroidThunkCpp_ShowConsoleWindow();
 extern bool AndroidThunkCpp_VirtualInputIgnoreClick(int, int);
 extern bool AndroidThunkCpp_IsVirtualKeyboardShown();
 extern bool AndroidThunkCpp_IsWebViewShown();
+extern bool AndroidThunkCpp_HasExtraIgnoreView();
 extern void AndroidThunkCpp_RestartApplication(const FString& IntentString);
 extern FString AndroidThunkCpp_GetIntentExtrasString(const FString& Key);
 
@@ -1236,7 +1237,7 @@ static int32_t HandleInputCB(struct android_app* app, AInputEvent* event)
 					return 0;
 				}
 			}
-			else if (AndroidThunkCpp_IsWebViewShown() && (type == TouchBegan || type == TouchMoved || type == TouchEnded))
+			else if ((AndroidThunkCpp_IsWebViewShown() || AndroidThunkCpp_HasExtraIgnoreView()) && (type == TouchBegan || type == TouchMoved || type == TouchEnded))
 			{
 				//ignore key down events when the the a web view is visible
 				if (TryIgnoreClick(event, actionPointer) && ((EventSource & 0x80) != 0x80))
diff --git a/Engine/Source/Runtime/Launch/Public/Android/AndroidJNI.h b/Engine/Source/Runtime/Launch/Public/Android/AndroidJNI.h
--- a/Engine/Source/Runtime/Launch/Public/Android/AndroidJNI.h
+++ b/Engine/Source/Runtime/Launch/Public/Android/AndroidJNI.h
@@ -90,9 +90,8 @@ public:
 	static jmethodID AndroidThunkJava_IsScreenCaptureDisabled;
 	static jmethodID AndroidThunkJava_DisableScreenCapture;
 
-	static jmethodID AndroidThunkCpp_VirtualInputIgnoreClick;
-	static jmethodID AndroidThunkCpp_IsVirtualKeyboardShown;
-	static jmethodID AndroidThunkCpp_IsWebViewShown;
+	// Native view click filtering
+	static jmethodID AndroidThunkJava_HasExtraIgnoreView;
 
 	// InputDeviceInfo member field ids
 	static jclass InputDeviceInfoClass;
