Fixed non-ASCII logging to stdout on Windows

Setting the console output code page to UTF-8 and using printf is the only approach that works reliably both with and without a console attached.

Alternatives that were considered:
- Using _setmode with either _O_U8TEXT or _O_WTEXT can work for wprintf but causes printf to assert and we cannot stop printf from being called.
- Using _wprintf_l or _printf_l with the .UTF-8 locale works without a console attached is incorrect when printing to a console launched by/with the debugger.
- Using printf with no other changes works in some circumstances but fails when launched by the debugger.
- Using CP_UTF8 with wprintf makes no difference to its behavior.

diff --git a/Engine/Source/Programs/AutomationTool/AutomationUtils/CommandletUtils.cs b/Engine/Source/Programs/AutomationTool/AutomationUtils/CommandletUtils.cs
--- a/Engine/Source/Programs/AutomationTool/AutomationUtils/CommandletUtils.cs
+++ b/Engine/Source/Programs/AutomationTool/AutomationUtils/CommandletUtils.cs
@@ -329,7 +329,6 @@ namespace AutomationTool
 			);
 			if (GlobalCommandLine.UTF8Output)
 			{
-				Args += " -UTF8Output";
 				RunOptions |= ERunOptions.UTF8Output;
 			}
 			RunResult = Run(EditorExe, Args, Options: RunOptions, Identifier: Commandlet, SpewFilterCallback: SpewFilterCallback);
diff --git a/Engine/Source/Programs/AutomationTool/BuildGraph/Tasks/PakFileTask.cs b/Engine/Source/Programs/AutomationTool/BuildGraph/Tasks/PakFileTask.cs
--- a/Engine/Source/Programs/AutomationTool/BuildGraph/Tasks/PakFileTask.cs
+++ b/Engine/Source/Programs/AutomationTool/BuildGraph/Tasks/PakFileTask.cs
@@ -152,10 +152,6 @@ namespace AutomationTool.Tasks
 			{
 				commandLine.Append(" -installed");
 			}
-			if (GlobalCommandLine.UTF8Output)
-			{
-				commandLine.AppendFormat(" -UTF8Output");
-			}
 
 			// Get the executable path
 			FileReference unrealPakExe;
diff --git a/Engine/Source/Runtime/Core/Private/Windows/WindowsPlatformMisc.cpp b/Engine/Source/Runtime/Core/Private/Windows/WindowsPlatformMisc.cpp
--- a/Engine/Source/Runtime/Core/Private/Windows/WindowsPlatformMisc.cpp
+++ b/Engine/Source/Runtime/Core/Private/Windows/WindowsPlatformMisc.cpp
@@ -1270,8 +1270,9 @@ bool FWindowsPlatformMisc::IsRemoteSession()
 
 void FWindowsPlatformMisc::SetUTF8Output()
 {
-	CA_SUPPRESS(6031)
-	_setmode(_fileno(stdout), _O_U8TEXT);
+	// Unicode support requires the UTF-8 code page and the use of printf.
+	// Note that this impacts stdout whether or not a console is attached.
+	SetConsoleOutputCP(CP_UTF8);
 }
 
 void FWindowsPlatformMisc::LocalPrint(const TCHAR* Message)
diff --git a/Engine/Source/Runtime/Launch/Private/LaunchEngineLoop.cpp b/Engine/Source/Runtime/Launch/Private/LaunchEngineLoop.cpp
--- a/Engine/Source/Runtime/Launch/Private/LaunchEngineLoop.cpp
+++ b/Engine/Source/Runtime/Launch/Private/LaunchEngineLoop.cpp
@@ -447,10 +447,11 @@ private:
 	#if PLATFORM_WINDOWS
 		if (bIsConsoleOutput)
 		{
-			return WriteLine<WIDECHAR>(V, Verbosity, Category, Time);
+			WriteLine<WIDECHAR>(V, Verbosity, Category, Time);
+			return;
 		}
 	#endif
-	#if PLATFORM_TCHAR_IS_UTF8CHAR || PLATFORM_TCHAR_IS_CHAR16
+	#if PLATFORM_WINDOWS || PLATFORM_TCHAR_IS_UTF8CHAR || PLATFORM_TCHAR_IS_CHAR16
 		WriteLine<UTF8CHAR>(V, Verbosity, Category, Time);
 	#else
 		WriteLine<WIDECHAR>(V, Verbosity, Category, Time);
@@ -640,10 +641,11 @@ private:
 	#if PLATFORM_WINDOWS
 		if (bIsConsoleOutput)
 		{
-			return WriteLine<WIDECHAR>(Object);
+			WriteLine<WIDECHAR>(Object);
+			return;
 		}
 	#endif
-	#if PLATFORM_TCHAR_IS_UTF8CHAR || PLATFORM_TCHAR_IS_CHAR16
+	#if PLATFORM_WINDOWS || PLATFORM_TCHAR_IS_UTF8CHAR || PLATFORM_TCHAR_IS_CHAR16
 		WriteLine<UTF8CHAR>(Object);
 	#else
 		WriteLine<WIDECHAR>(Object);
@@ -1973,10 +1975,7 @@ int32 FEngineLoop::PreInitPreStartupScreen(const TCHAR* CmdLine)
 
 	FMemory::SetupTLSCachesOnCurrentThread();
 
-	if (FParse::Param(CmdLine, TEXT("UTF8Output")))
-	{
-		FPlatformMisc::SetUTF8Output();
-	}
+	FPlatformMisc::SetUTF8Output();
 
 #if !UE_BUILD_SHIPPING
 	if (FParse::Param(CmdLine, TEXT("IgnoreDebugger")))
@@ -2097,6 +2096,12 @@ int32 FEngineLoop::PreInitPreStartupScreen(const TCHAR* CmdLine)
 	}
 #endif
 
+	if (FParse::Param(CmdLine, TEXT("UTF8Output")))
+	{
+		// UE_DEPRECATED(5.6, "This comment exists to describe when -UTF8Output was deprecated and allow finding it in a search for UE 5.6 deprecations.")
+		UE_LOG(LogInit, Display, TEXT("Output is UTF-8 by default. Use of -UTF8Output on the command is deprecated and can be removed."));
+	}
+
 // If we are in Debug, Development, or Test/Shipping with ALLOW_PROFILEGPU... enabled, then automatically allow draw events.
 // Command lines allow disabling of draw events if WITH_PROGILEGPU is enabled. Test/Shipping on their own require explicit opt in.
 #if WITH_PROFILEGPU	
