Removed unnecessary ensure in FAndroidPlatformBackgroundHttpManager::FindRequestByID. Fixed OS thread always thinking it's the game thread.

#android
#tests tested the background downloader to make sure it doesn't assert and deadlock

diff --git a/Engine/Plugins/Online/Android/AndroidFetchBackgroundDownload/Source/AndroidFetchBackgroundDownload/Private/AndroidPlatformBackgroundHttpManager.cpp b/Engine/Plugins/Online/Android/AndroidFetchBackgroundDownload/Source/AndroidFetchBackgroundDownload/Private/AndroidPlatformBackgroundHttpManager.cpp
--- a/Engine/Plugins/Online/Android/AndroidFetchBackgroundDownload/Source/AndroidFetchBackgroundDownload/Private/AndroidPlatformBackgroundHttpManager.cpp
+++ b/Engine/Plugins/Online/Android/AndroidFetchBackgroundDownload/Source/AndroidFetchBackgroundDownload/Private/AndroidPlatformBackgroundHttpManager.cpp
@@ -508,8 +508,6 @@ FAndroidBackgroundHttpRequestPtr FAndroidPlatformBackgroundHttpManager::FindRequ
 		}
 	}
 
-	ensureAlwaysMsgf(ReturnedPtr.IsValid(), TEXT("No matching valid request found for %s"), *RequestID);
-
 	return ReturnedPtr;
 }
 
diff --git a/Engine/Source/Runtime/Launch/Private/Android/AndroidJNI.cpp b/Engine/Source/Runtime/Launch/Private/Android/AndroidJNI.cpp
--- a/Engine/Source/Runtime/Launch/Private/Android/AndroidJNI.cpp
+++ b/Engine/Source/Runtime/Launch/Private/Android/AndroidJNI.cpp
@@ -1235,21 +1235,11 @@ JNI_METHOD void Java_com_epicgames_unreal_GameActivity_nativeVirtualKeyboardChan
 
 JNI_METHOD void Java_com_epicgames_unreal_GameActivity_nativeVirtualKeyboardSendKey(JNIEnv* jenv, jobject thiz, jint keyCode)
 {
-	if (FTaskGraphInterface::IsRunning())
-	{
-		// Dispatch when ready to delay the submit when the text is being changed by another task (e.g Java_com_epicgames_unreal_GameActivity_nativeVirtualKeyboardChanged)
-		FFunctionGraphTask::CreateAndDispatchWhenReady([keyCode = keyCode, CurrentVirtualKeyboardWidget = VirtualKeyboardWidget]()
-		{
-			if (CurrentVirtualKeyboardWidget.IsValid() && CurrentVirtualKeyboardWidget == VirtualKeyboardWidget)
-			{
-				FDeferredAndroidMessage Message;
+	FDeferredAndroidMessage Message;
 
-				Message.messageType = MessageType_KeyDown;
-				Message.KeyEventData.keyId = keyCode;
-				FAndroidInputInterface::DeferMessage(Message);
-			}
-		}, TStatId(), NULL, ENamedThreads::GameThread);
-	}
+	Message.messageType = MessageType_KeyDown;
+	Message.KeyEventData.keyId = keyCode;
+	FAndroidInputInterface::DeferMessage(Message);
 }
 
 JNI_METHOD void Java_com_epicgames_unreal_GameActivity_nativeVirtualKeyboardSendSelection(JNIEnv* jenv, jobject thiz, jint selStart, jint selEnd)
@@ -1944,6 +1934,8 @@ void AndroidThunkCpp_ReleaseWifiManagerMulticastLock()
 
 JNIEXPORT jint JNI_OnLoad(JavaVM* InJavaVM, void* InReserved)
 {
+	FTaskTagScope::SetTagNone();
+	
 	FPlatformMisc::LowLevelOutputDebugString(TEXT("In the JNI_OnLoad function"));
 
 	JNIEnv* Env = NULL;
