Undo bad merge with changelist 40732326

#android

diff --git a/Engine/Source/Runtime/Launch/Private/Android/AndroidJNI.cpp b/Engine/Source/Runtime/Launch/Private/Android/AndroidJNI.cpp
--- a/Engine/Source/Runtime/Launch/Private/Android/AndroidJNI.cpp
+++ b/Engine/Source/Runtime/Launch/Private/Android/AndroidJNI.cpp
@@ -1235,11 +1235,21 @@ JNI_METHOD void Java_com_epicgames_unreal_GameActivity_nativeVirtualKeyboardChan
 
 JNI_METHOD void Java_com_epicgames_unreal_GameActivity_nativeVirtualKeyboardSendKey(JNIEnv* jenv, jobject thiz, jint keyCode)
 {
-	FDeferredAndroidMessage Message;
+	if (FTaskGraphInterface::IsRunning())
+	{
+		// Dispatch when ready to delay the submit when the text is being changed by another task (e.g Java_com_epicgames_unreal_GameActivity_nativeVirtualKeyboardChanged)
+		FFunctionGraphTask::CreateAndDispatchWhenReady([keyCode = keyCode, CurrentVirtualKeyboardWidget = VirtualKeyboardWidget]()
+		{
+			if (CurrentVirtualKeyboardWidget.IsValid() && CurrentVirtualKeyboardWidget == VirtualKeyboardWidget)
+			{
+				FDeferredAndroidMessage Message;
 
-	Message.messageType = MessageType_KeyDown;
-	Message.KeyEventData.keyId = keyCode;
-	FAndroidInputInterface::DeferMessage(Message);
+				Message.messageType = MessageType_KeyDown;
+				Message.KeyEventData.keyId = keyCode;
+				FAndroidInputInterface::DeferMessage(Message);
+			}
+		}, TStatId(), NULL, ENamedThreads::GameThread);
+	}
 }
 
 JNI_METHOD void Java_com_epicgames_unreal_GameActivity_nativeVirtualKeyboardSendSelection(JNIEnv* jenv, jobject thiz, jint selStart, jint selEnd)
