Logging: Disable the backlog only after module loading is complete

This gives modules a chance to consume the backlog if they construct their own output devices.

diff --git a/Engine/Source/Editor/UnrealEd/Private/EditorEngine.cpp b/Engine/Source/Editor/UnrealEd/Private/EditorEngine.cpp
--- a/Engine/Source/Editor/UnrealEd/Private/EditorEngine.cpp
+++ b/Engine/Source/Editor/UnrealEd/Private/EditorEngine.cpp
@@ -1327,8 +1327,6 @@ void UEditorEngine::Init(IEngineLoop* InEngineLoop)
 	}
 	UModel::SetGlobalBSPTexelScale(BSPTexelScale);
 
-	GLog->EnableBacklog( false );
-
 	// Load game user settings and apply
 	UGameUserSettings* MyGameUserSettings = GetGameUserSettings();
 	if (MyGameUserSettings)
diff --git a/Engine/Source/Runtime/Launch/Private/LaunchEngineLoop.cpp b/Engine/Source/Runtime/Launch/Private/LaunchEngineLoop.cpp
--- a/Engine/Source/Runtime/Launch/Private/LaunchEngineLoop.cpp
+++ b/Engine/Source/Runtime/Launch/Private/LaunchEngineLoop.cpp
@@ -2928,12 +2928,6 @@ int32 FEngineLoop::PreInitPreStartupScreen(const TCHAR* CmdLine)
 	}
 #endif // WITH_ENGINE
 
-	// If we're not in the editor stop collecting the backlog now that we know
-	if (!GIsEditor)
-	{
-		GLog->EnableBacklog(false);
-	}
-
 #if WITH_ENGINE && FRAMEPRO_ENABLED
 	FFrameProProfiler::Initialize();
 #endif // FRAMEPRO_ENABLED
@@ -4929,6 +4923,10 @@ int32 FEngineLoop::Init()
 	// Call module loading phases completion callbacks
 	SetEngineStartupModuleLoadingComplete();
 
+	// Disable the backlog only after module loading is complete because that gives modules
+	// a chance to consume the backlog when constructing output devices.
+	GLog->EnableBacklog(false);
+
 	{
 		SCOPED_BOOT_TIMING("GEngine->Start()");
 		GEngine->Start();
