Don't try to change the CWD on startup if we are disabling it to avoid a check() in the SetCurrentWorkingDirectoryToBaseDir implmentations.

#rnx

diff --git a/Engine/Source/Runtime/Launch/Private/LaunchEngineLoop.cpp b/Engine/Source/Runtime/Launch/Private/LaunchEngineLoop.cpp
--- a/Engine/Source/Runtime/Launch/Private/LaunchEngineLoop.cpp
+++ b/Engine/Source/Runtime/Launch/Private/LaunchEngineLoop.cpp
@@ -1989,7 +1989,9 @@ int32 FEngineLoop::PreInitPreStartupScreen(const TCHAR* CmdLine)
 #endif // !UE_BUILD_SHIPPING
 
 	// Switch into executable's directory.
+#if !defined(DISABLE_CWD_CHANGES) || DISABLE_CWD_CHANGES==0
 	FPlatformProcess::SetCurrentWorkingDirectoryToBaseDir();
+#endif
 
 	// this is set later with shorter command lines, but we want to make sure it is set ASAP as some subsystems will do the tests themselves...
 	// also realize that command lines can be pulled from the network at a slightly later time.
@@ -2147,7 +2149,9 @@ int32 FEngineLoop::PreInitPreStartupScreen(const TCHAR* CmdLine)
 #endif // !UE_BUILD_SHIPPING
 
 	// Switch into executable's directory (may be required by some of the platform file overrides)
+#if !defined(DISABLE_CWD_CHANGES) || DISABLE_CWD_CHANGES==0
 	FPlatformProcess::SetCurrentWorkingDirectoryToBaseDir();
+#endif
 
 	// This fixes up the relative project path, needs to happen before we set platform file paths
 	if (FPlatformProperties::IsProgram() == false)
@@ -6459,7 +6463,9 @@ bool FEngineLoop::AppInit( )
 	GSystemStartTime = FDateTime::Now().ToString();
 
 	// Switch into executable's directory.
+#if !defined(DISABLE_CWD_CHANGES) || DISABLE_CWD_CHANGES==0
 	FPlatformProcess::SetCurrentWorkingDirectoryToBaseDir();
+#endif
 
 	{
 		SCOPED_BOOT_TIMING("IFileManager::Get().ProcessCommandLineOptions()");
