Fix iOS dialog alerts not showing during startup

diff --git a/Engine/Source/Runtime/Launch/Private/IOS/IOSAppDelegateConsoleHandling.cpp b/Engine/Source/Runtime/Launch/Private/IOS/IOSAppDelegateConsoleHandling.cpp
--- a/Engine/Source/Runtime/Launch/Private/IOS/IOSAppDelegateConsoleHandling.cpp
+++ b/Engine/Source/Runtime/Launch/Private/IOS/IOSAppDelegateConsoleHandling.cpp
@@ -138,14 +138,25 @@ extern bool GDebugConsoleOpen;
  */
 - (void)ShowAlert:(NSMutableArray*)StringArray
 {
+	UIViewController *ViewController = [IOSAppDelegate GetDelegate].IOSController;
 	if (GShowSplashScreen)
 	{
 		if ([[IOSAppDelegate GetDelegate].Window viewWithTag : 200] != nil)
 		{
-			[[[IOSAppDelegate GetDelegate].Window viewWithTag : 200] removeFromSuperview];
+			// If the splash screen is showing, the UIView isn't valid yet, so check for a ViewController
+			if(ViewController == nil)
+			{
+				// walk the responder chain until we get to a VC
+				UIResponder *Responder = [[IOSAppDelegate GetDelegate].Window viewWithTag : 200];
+				while (Responder != nil && ![Responder isKindOfClass:[UIViewController class]])
+				{
+					Responder = [Responder nextResponder];
+				}
+				ViewController = (UIViewController*)Responder;
+			}
 		}
-		GShowSplashScreen = false;
 	}
+	
 	if ([UIAlertController class])
 	{
 		UIAlertController* AlertController = [UIAlertController alertControllerWithTitle:[StringArray objectAtIndex : 0]
@@ -168,7 +179,15 @@ extern bool GDebugConsoleOpen;
 			[AlertController addAction : alertAction];
 		}
 
-		[[IOSAppDelegate GetDelegate].IOSController presentViewController : AlertController animated : YES completion : nil];
+		NSLog(@"Showing Alert:\n  Title:%@\n  Message:%@", [StringArray objectAtIndex:0], [StringArray objectAtIndex:1]);
+		if(ViewController != nil)
+		{
+			[ViewController presentViewController : AlertController animated : YES completion : nil];
+		}
+		else
+		{
+			NSLog(@"No valid ViewController to show Alert");
+		}
 	}
 }
 
