Logging: Extracted FOutputDeviceStdOut from LaunchEngineLoop into Core

diff --git a/Engine/Source/Runtime/Core/Private/Misc/OutputDeviceStdOut.cpp b/Engine/Source/Runtime/Core/Private/Misc/OutputDeviceStdOut.cpp
new file mode 100644
--- /dev/null
+++ b/Engine/Source/Runtime/Core/Private/Misc/OutputDeviceStdOut.cpp
@@ -0,0 +1,351 @@
+// Copyright Epic Games, Inc. All Rights Reserved.
+
+#include "Misc/OutputDeviceStdOut.h"
+
+#include "CoreGlobals.h"
+#include "Logging/StructuredLog.h"
+#include "Misc/AsciiSet.h"
+#include "Misc/CommandLine.h"
+#include "Misc/DateTime.h"
+#include "Misc/OutputDeviceHelper.h"
+#include "Misc/Parse.h"
+#include "Serialization/CompactBinary.h"
+#include "Serialization/CompactBinarySerialization.h"
+#include "Serialization/CompactBinaryWriter.h"
+
+#if PLATFORM_WINDOWS
+#include "Windows/WindowsHWrapper.h"
+#endif
+
+FOutputDeviceStdOutput::FOutputDeviceStdOutput()
+{
+#if PLATFORM_WINDOWS
+	bIsConsoleOutput = IsStdoutAttachedToConsole() && !FParse::Param(FCommandLine::Get(), TEXT("GenericConsoleOutput"));
+#endif
+
+	bIsJsonOutput = FParse::Param(FCommandLine::Get(), TEXT("JsonStdOut"));
+	if (!bIsJsonOutput)
+	{
+		if (FString Env = FPlatformMisc::GetEnvironmentVariable(TEXT("UE_LOG_JSON_TO_STDOUT")); !Env.IsEmpty())
+		{
+			bIsJsonOutput = FCString::Atoi(*Env) != 0;
+		}
+	}
+
+	if (FParse::Param(FCommandLine::Get(), TEXT("AllowStdOutLogVerbosity")))
+	{
+		AllowedLogVerbosity = ELogVerbosity::Log;
+	}
+
+	if (FParse::Param(FCommandLine::Get(), TEXT("FullStdOutLogOutput")))
+	{
+		AllowedLogVerbosity = ELogVerbosity::All;
+	}
+
+	if (stdout == nullptr)
+	{
+		AllowedLogVerbosity = ELogVerbosity::NoLogging;
+	}
+}
+
+void FOutputDeviceStdOutput::Serialize(const TCHAR* V, ELogVerbosity::Type Verbosity, const FName& Category, double Time)
+{
+	if (Verbosity <= AllowedLogVerbosity)
+	{
+		return bIsJsonOutput ? SerializeAsJson(V, Verbosity, Category, Time) : SerializeAsText(V, Verbosity, Category, Time);
+	}
+}
+
+void FOutputDeviceStdOutput::SerializeRecord(const UE::FLogRecord& Record)
+{
+	if (Record.GetVerbosity() <= AllowedLogVerbosity)
+	{
+		return bIsJsonOutput ? SerializeRecordAsJson(Record) : SerializeRecordAsText(Record);
+	}
+}
+
+// Several functions below are FORCENOINLINE to reduce total required stack space by limiting
+// the scope of string builders and compact binary writers.
+
+FORCENOINLINE void FOutputDeviceStdOutput::SerializeAsText(const TCHAR* V, ELogVerbosity::Type Verbosity, const FName& Category, double Time)
+{
+#if PLATFORM_WINDOWS
+	if (bIsConsoleOutput)
+	{
+		WriteLine<WIDECHAR>(V, Verbosity, Category, Time);
+		return;
+	}
+#endif
+#if PLATFORM_WINDOWS || PLATFORM_TCHAR_IS_UTF8CHAR || PLATFORM_TCHAR_IS_CHAR16
+	WriteLine<UTF8CHAR>(V, Verbosity, Category, Time);
+#else
+	WriteLine<WIDECHAR>(V, Verbosity, Category, Time);
+#endif
+}
+
+FORCENOINLINE void FOutputDeviceStdOutput::SerializeRecordAsText(const UE::FLogRecord& Record)
+{
+	TStringBuilder<512> V;
+	Record.FormatMessageTo(V);
+	Serialize(*V, Record.GetVerbosity(), Record.GetCategory(), -1.0);
+}
+
+template <typename CharType>
+FORCENOINLINE void FOutputDeviceStdOutput::WriteLine(const TCHAR* V, ELogVerbosity::Type Verbosity, const FName& Category, double Time)
+{
+	TStringBuilderWithBuffer<CharType, 512> Line;
+	FOutputDeviceHelper::AppendFormatLogLine(Line, Verbosity, Category, V, GPrintLogTimes, Time);
+	Line.AppendChar('\n');
+	WriteLine(Line);
+}
+
+FORCENOINLINE void FOutputDeviceStdOutput::SerializeAsJson(const TCHAR* V, ELogVerbosity::Type Verbosity, const FName& Category, double Time)
+{
+	const bool bShowCategory = GPrintLogCategory && !Category.IsNone();
+	const bool bShowVerbosity = GPrintLogVerbosity && (Verbosity & ELogVerbosity::VerbosityMask) != ELogVerbosity::Log;
+
+	TCbWriter<1024> Writer;
+	Writer.BeginObject();
+	Writer.AddDateTime(ANSITEXTVIEW("time"), FDateTime::UtcNow());
+	Writer.AddString(ANSITEXTVIEW("level"), GetLevel(Verbosity));
+	AddMessage(Writer, V, Verbosity, Category, bShowCategory, bShowVerbosity);
+	if (bShowCategory || bShowVerbosity)
+	{
+		AddFormat(Writer, V, bShowCategory, bShowVerbosity);
+
+		Writer.BeginObject(ANSITEXTVIEW("properties"));
+		if (bShowCategory)
+		{
+			Writer.BeginObject(ANSITEXTVIEW("_channel"));
+			Writer.AddString(ANSITEXTVIEW("$type"), ANSITEXTVIEW("Channel"));
+			Writer.AddString(ANSITEXTVIEW("$text"), WriteToUtf8String<64>(Category));
+			Writer.EndObject();
+		}
+		if (bShowVerbosity)
+		{
+			Writer.BeginObject(ANSITEXTVIEW("_severity"));
+			Writer.AddString(ANSITEXTVIEW("$type"), ANSITEXTVIEW("Severity"));
+			Writer.AddString(ANSITEXTVIEW("$text"), ToString(Verbosity));
+			Writer.EndObject();
+		}
+		Writer.EndObject();
+	}
+	Writer.EndObject();
+
+	WriteAsJson(Writer);
+}
+
+FORCENOINLINE void FOutputDeviceStdOutput::AddMessage(FCbWriter& Writer, const TCHAR* V, ELogVerbosity::Type Verbosity, const FName& Category, bool bShowCategory, bool bShowVerbosity)
+{
+	TUtf8StringBuilder<512> Message;
+	if (bShowCategory)
+	{
+		Message << Category << ANSITEXTVIEW(": ");
+	}
+	if (bShowVerbosity)
+	{
+		Message << ToString(Verbosity) << ANSITEXTVIEW(": ");
+	}
+	Message << V;
+	Writer.AddString(ANSITEXTVIEW("message"), Message);
+}
+
+FORCENOINLINE void FOutputDeviceStdOutput::AddFormat(FCbWriter& Writer, const TCHAR* Message, bool bShowCategory, bool bShowVerbosity)
+{
+	TUtf8StringBuilder<512> Format;
+	if (bShowCategory)
+	{
+		Format << ANSITEXTVIEW("{_channel}: ");
+	}
+	if (bShowVerbosity)
+	{
+		Format << ANSITEXTVIEW("{_severity}: ");
+	}
+
+	// Escape {} in Message
+	for (constexpr FAsciiSet Brackets("{}");;)
+	{
+		const TCHAR* End = FAsciiSet::FindFirstOrEnd(Message, Brackets);
+		Format.Append(Message, UE_PTRDIFF_TO_INT32(End - Message));
+		if (!*End)
+		{
+			break;
+		}
+		Format.AppendChar(*End);
+		Format.AppendChar(*End);
+		Message = End + 1;
+	}
+
+	Writer.AddString(ANSITEXTVIEW("format"), Format);
+}
+
+FORCENOINLINE void FOutputDeviceStdOutput::SerializeRecordAsJson(const UE::FLogRecord& Record)
+{
+	const bool bShowCategory = GPrintLogCategory && !Record.GetCategory().IsNone();
+	const bool bShowVerbosity = GPrintLogVerbosity && (Record.GetVerbosity() & ELogVerbosity::VerbosityMask) != ELogVerbosity::Log;
+
+	TCbWriter<1024> Writer;
+	Writer.BeginObject();
+	Writer.AddDateTime(ANSITEXTVIEW("time"), Record.GetTime().GetUtcTime());
+	Writer.AddString(ANSITEXTVIEW("level"), GetLevel(Record.GetVerbosity()));
+	AddMessage(Writer, Record, bShowCategory, bShowVerbosity);
+	if (bShowCategory || bShowVerbosity || Record.GetFields())
+	{
+		AddFormat(Writer, Record, bShowCategory, bShowVerbosity);
+
+		Writer.BeginObject(ANSITEXTVIEW("properties"));
+		if (bShowCategory)
+		{
+			Writer.BeginObject(ANSITEXTVIEW("_channel"));
+			Writer.AddString(ANSITEXTVIEW("$type"), ANSITEXTVIEW("Channel"));
+			Writer.AddString(ANSITEXTVIEW("$text"), WriteToUtf8String<64>(Record.GetCategory()));
+			Writer.EndObject();
+		}
+		if (bShowVerbosity)
+		{
+			Writer.BeginObject(ANSITEXTVIEW("_severity"));
+			Writer.AddString(ANSITEXTVIEW("$type"), ANSITEXTVIEW("Severity"));
+			Writer.AddString(ANSITEXTVIEW("$text"), ToString(Record.GetVerbosity()));
+			Writer.EndObject();
+		}
+		if (const TCHAR* TextNamespace = Record.GetTextNamespace())
+		{
+			Writer.AddString(ANSITEXTVIEW("_ns"), TextNamespace);
+		}
+		if (const TCHAR* TextKey = Record.GetTextKey())
+		{
+			Writer.AddString(ANSITEXTVIEW("_key"), TextKey);
+		}
+		for (const FCbField& Field : Record.GetFields())
+		{
+			Writer.AddField(Field.GetName(), Field);
+		}
+		Writer.EndObject();
+	}
+	Writer.EndObject();
+
+	WriteAsJson(Writer);
+}
+
+FORCENOINLINE void FOutputDeviceStdOutput::AddMessage(FCbWriter& Writer, const UE::FLogRecord& Record, bool bShowCategory, bool bShowVerbosity)
+{
+	TUtf8StringBuilder<512> Message;
+	if (bShowCategory)
+	{
+		Message << Record.GetCategory() << ANSITEXTVIEW(": ");
+	}
+	if (bShowVerbosity)
+	{
+		Message << ToString(Record.GetVerbosity()) << ANSITEXTVIEW(": ");
+	}
+	Record.FormatMessageTo(Message);
+	Writer.AddString(ANSITEXTVIEW("message"), Message);
+}
+
+FORCENOINLINE void FOutputDeviceStdOutput::AddFormat(FCbWriter& Writer, const UE::FLogRecord& Record, bool bShowCategory, bool bShowVerbosity)
+{
+	TUtf8StringBuilder<512> Format;
+	if (bShowCategory)
+	{
+		Format << ANSITEXTVIEW("{_channel}: ");
+	}
+	if (bShowVerbosity)
+	{
+		Format << ANSITEXTVIEW("{_severity}: ");
+	}
+	Format.Append(Record.GetFormat());
+	Writer.AddString(ANSITEXTVIEW("format"), Format);
+}
+
+void FOutputDeviceStdOutput::WriteAsJson(const FCbWriter& Writer)
+{
+	TArray<uint8, TInlineAllocator64<512>> Buffer;
+	Buffer.AddUninitialized((int64)Writer.GetSaveSize());
+	FCbFieldView Object = Writer.Save(MakeMemoryView(Buffer));
+
+#if PLATFORM_WINDOWS
+	if (bIsConsoleOutput)
+	{
+		WriteLine<WIDECHAR>(Object);
+		return;
+	}
+#endif
+#if PLATFORM_WINDOWS || PLATFORM_TCHAR_IS_UTF8CHAR || PLATFORM_TCHAR_IS_CHAR16
+	WriteLine<UTF8CHAR>(Object);
+#else
+	WriteLine<WIDECHAR>(Object);
+#endif
+}
+
+template <typename CharType>
+FORCENOINLINE void FOutputDeviceStdOutput::WriteLine(const FCbFieldView& Field)
+{
+	TStringBuilderWithBuffer<CharType, 512> Line;
+	CompactBinaryToCompactJson(Field, Line);
+	Line.AppendChar('\n');
+	WriteLine(Line);
+}
+
+void FOutputDeviceStdOutput::WriteLine(FUtf8StringBuilderBase& Line) const
+{
+	printf("%s", (const char*)*Line);
+	fflush(stdout);
+}
+
+#if PLATFORM_WINDOWS || !(PLATFORM_TCHAR_IS_UTF8CHAR || PLATFORM_TCHAR_IS_CHAR16)
+void FOutputDeviceStdOutput::WriteLine(FWideStringBuilderBase& Line) const
+{
+#if PLATFORM_WINDOWS
+	if (bIsConsoleOutput)
+	{
+		WriteConsoleW(GetStdHandle(STD_OUTPUT_HANDLE), *Line, Line.Len(), nullptr, nullptr);
+		return;
+	}
+#endif
+
+#if PLATFORM_USE_LS_SPEC_FOR_WIDECHAR
+	// printf prints wchar_t strings just fine with %ls, while mixing printf()/wprintf() is not recommended (see https://stackoverflow.com/questions/8681623/printf-and-wprintf-in-single-c-code)
+	printf("%ls", *Line);
+#else
+	wprintf(TEXT("%s"), *Line);
+#endif
+	fflush(stdout);
+}
+#endif
+
+FAnsiStringView FOutputDeviceStdOutput::GetLevel(ELogVerbosity::Type Verbosity)
+{
+	switch (Verbosity & ELogVerbosity::VerbosityMask)
+	{
+	case ELogVerbosity::Fatal:
+		return ANSITEXTVIEW("Critical");
+	case ELogVerbosity::Error:
+		return ANSITEXTVIEW("Error");
+	case ELogVerbosity::Warning:
+		return ANSITEXTVIEW("Warning");
+	case ELogVerbosity::Display:
+	case ELogVerbosity::Log:
+	default:
+		return ANSITEXTVIEW("Information");
+	case ELogVerbosity::Verbose:
+	case ELogVerbosity::VeryVerbose:
+		return ANSITEXTVIEW("Debug");
+	}
+}
+
+#if PLATFORM_WINDOWS
+bool FOutputDeviceStdOutput::IsStdoutAttachedToConsole()
+{
+	HANDLE StdoutHandle = GetStdHandle(STD_OUTPUT_HANDLE);
+	if (StdoutHandle != INVALID_HANDLE_VALUE)
+	{
+		DWORD FileType = GetFileType(StdoutHandle);
+		if (FileType == FILE_TYPE_CHAR)
+		{
+			return true;
+		}
+	}
+
+	return false;
+}
+#endif
diff --git a/Engine/Source/Runtime/Core/Public/Misc/OutputDeviceStdOut.h b/Engine/Source/Runtime/Core/Public/Misc/OutputDeviceStdOut.h
new file mode 100644
--- /dev/null
+++ b/Engine/Source/Runtime/Core/Public/Misc/OutputDeviceStdOut.h
@@ -0,0 +1,64 @@
+// Copyright Epic Games, Inc. All Rights Reserved.
+
+#pragma once
+
+#include "Containers/StringFwd.h"
+#include "Misc/OutputDevice.h"
+
+class FCbFieldView;
+class FCbWriter;
+
+class FOutputDeviceStdOutput final : public FOutputDevice
+{
+public:
+	CORE_API FOutputDeviceStdOutput();
+
+	bool CanBeUsedOnAnyThread() const final { return true; }
+	bool CanBeUsedOnPanicThread() const final { return true; }
+
+	void Serialize(const TCHAR* V, ELogVerbosity::Type Verbosity, const FName& Category) final
+	{
+		Serialize(V, Verbosity, Category, -1.0);
+	}
+
+	void Serialize(const TCHAR* V, ELogVerbosity::Type Verbosity, const FName& Category, double Time) final;
+	void SerializeRecord(const UE::FLogRecord& Record) final;
+
+private:
+	void SerializeAsText(const TCHAR* V, ELogVerbosity::Type Verbosity, const FName& Category, double Time);
+	void SerializeRecordAsText(const UE::FLogRecord& Record);
+
+	template <typename CharType>
+	void WriteLine(const TCHAR* V, ELogVerbosity::Type Verbosity, const FName& Category, double Time);
+
+	void SerializeAsJson(const TCHAR* V, ELogVerbosity::Type Verbosity, const FName& Category, double Time);
+
+	static void AddMessage(FCbWriter& Writer, const TCHAR* V, ELogVerbosity::Type Verbosity, const FName& Category, bool bShowCategory, bool bShowVerbosity);
+	static void AddFormat(FCbWriter& Writer, const TCHAR* Message, bool bShowCategory, bool bShowVerbosity);
+
+	void SerializeRecordAsJson(const UE::FLogRecord& Record);
+
+	static void AddMessage(FCbWriter& Writer, const UE::FLogRecord& Record, bool bShowCategory, bool bShowVerbosity);
+	static void AddFormat(FCbWriter& Writer, const UE::FLogRecord& Record, bool bShowCategory, bool bShowVerbosity);
+
+	void WriteAsJson(const FCbWriter& Writer);
+
+	template <typename CharType>
+	void WriteLine(const FCbFieldView& Field);
+
+	void WriteLine(FUtf8StringBuilderBase& Line) const;
+#if PLATFORM_WINDOWS || !(PLATFORM_TCHAR_IS_UTF8CHAR || PLATFORM_TCHAR_IS_CHAR16)
+	void WriteLine(FWideStringBuilderBase& Line) const;
+#endif
+
+	static FAnsiStringView GetLevel(ELogVerbosity::Type Verbosity);
+
+private:
+	ELogVerbosity::Type AllowedLogVerbosity = ELogVerbosity::Display;
+	bool bIsConsoleOutput = false;
+	bool bIsJsonOutput = false;
+
+#if PLATFORM_WINDOWS
+	static bool IsStdoutAttachedToConsole();
+#endif
+};
diff --git a/Engine/Source/Runtime/Launch/Private/LaunchEngineLoop.cpp b/Engine/Source/Runtime/Launch/Private/LaunchEngineLoop.cpp
--- a/Engine/Source/Runtime/Launch/Private/LaunchEngineLoop.cpp
+++ b/Engine/Source/Runtime/Launch/Private/LaunchEngineLoop.cpp
@@ -1,3 +1,4 @@
+
 // Copyright Epic Games, Inc. All Rights Reserved.
 
 #include "LaunchEngineLoop.h"
@@ -34,6 +35,7 @@
 #include "Misc/CommandLine.h"
 #include "Misc/App.h"
 #include "Misc/OutputDeviceConsole.h"
+#include "Misc/OutputDeviceStdOut.h"
 #include "Misc/ScopeExit.h"
 #include "Misc/StringBuilder.h"
 #include "Misc/TrackedActivity.h"
@@ -382,360 +384,6 @@ private:
 	FEvent* CompleteEvent;
 };
 
-// Pipe output to std output
-// This enables UBT to collect the output for it's own use
-class FOutputDeviceStdOutput final : public FOutputDevice
-{
-public:
-	FOutputDeviceStdOutput()
-	{
-#if PLATFORM_WINDOWS
-		bIsConsoleOutput = IsStdoutAttachedToConsole() && !FParse::Param(FCommandLine::Get(), TEXT("GenericConsoleOutput"));
-#endif
-
-		bIsJsonOutput = FParse::Param(FCommandLine::Get(), TEXT("JsonStdOut"));
-		if (!bIsJsonOutput)
-		{
-			if (FString Env = FPlatformMisc::GetEnvironmentVariable(TEXT("UE_LOG_JSON_TO_STDOUT")); !Env.IsEmpty())
-			{
-				bIsJsonOutput = FCString::Atoi(*Env) != 0;
-			}
-		}
-
-		if (FParse::Param(FCommandLine::Get(), TEXT("AllowStdOutLogVerbosity")))
-		{
-			AllowedLogVerbosity = ELogVerbosity::Log;
-		}
-
-		if (FParse::Param(FCommandLine::Get(), TEXT("FullStdOutLogOutput")))
-		{
-			AllowedLogVerbosity = ELogVerbosity::All;
-		}
-
-		if (stdout == nullptr)
-		{
-			AllowedLogVerbosity = ELogVerbosity::NoLogging;
-		}
-	}
-
-	bool CanBeUsedOnAnyThread() const final { return true; }
-	bool CanBeUsedOnPanicThread() const final { return true; }
-
-	void Serialize(const TCHAR* V, ELogVerbosity::Type Verbosity, const FName& Category) final
-	{
-		Serialize(V, Verbosity, Category, -1.0);
-	}
-
-	void Serialize(const TCHAR* V, ELogVerbosity::Type Verbosity, const FName& Category, double Time) final
-	{
-		if (Verbosity <= AllowedLogVerbosity)
-		{
-			return bIsJsonOutput ? SerializeAsJson(V, Verbosity, Category, Time) : SerializeAsText(V, Verbosity, Category, Time);
-		}
-	}
-
-	void SerializeRecord(const UE::FLogRecord& Record) final
-	{
-		if (Record.GetVerbosity() <= AllowedLogVerbosity)
-		{
-			return bIsJsonOutput ? SerializeRecordAsJson(Record) : SerializeRecordAsText(Record);
-		}
-	}
-
-private:
-	// Several functions below are FORCENOINLINE to reduce total required stack space by limiting
-	// the scope of string builders and compact binary writers.
-
-	FORCENOINLINE void SerializeAsText(const TCHAR* V, ELogVerbosity::Type Verbosity, const FName& Category, double Time)
-	{
-	#if PLATFORM_WINDOWS
-		if (bIsConsoleOutput)
-		{
-			WriteLine<WIDECHAR>(V, Verbosity, Category, Time);
-			return;
-		}
-	#endif
-	#if PLATFORM_WINDOWS || PLATFORM_TCHAR_IS_UTF8CHAR || PLATFORM_TCHAR_IS_CHAR16
-		WriteLine<UTF8CHAR>(V, Verbosity, Category, Time);
-	#else
-		WriteLine<WIDECHAR>(V, Verbosity, Category, Time);
-	#endif
-	}
-
-	FORCENOINLINE void SerializeRecordAsText(const UE::FLogRecord& Record)
-	{
-		TStringBuilder<512> V;
-		Record.FormatMessageTo(V);
-		Serialize(*V, Record.GetVerbosity(), Record.GetCategory(), -1.0);
-	}
-
-	template <typename CharType>
-	FORCENOINLINE void WriteLine(const TCHAR* V, ELogVerbosity::Type Verbosity, const FName& Category, double Time)
-	{
-		TStringBuilderWithBuffer<CharType, 512> Line;
-		FOutputDeviceHelper::AppendFormatLogLine(Line, Verbosity, Category, V, GPrintLogTimes, Time);
-		Line.AppendChar('\n');
-		WriteLine(Line);
-	}
-
-	FORCENOINLINE void SerializeAsJson(const TCHAR* V, ELogVerbosity::Type Verbosity, const FName& Category, double Time)
-	{
-		const bool bShowCategory = GPrintLogCategory && !Category.IsNone();
-		const bool bShowVerbosity = GPrintLogVerbosity && (Verbosity & ELogVerbosity::VerbosityMask) != ELogVerbosity::Log;
-
-		TCbWriter<1024> Writer;
-		Writer.BeginObject();
-		Writer.AddDateTime(ANSITEXTVIEW("time"), FDateTime::UtcNow());
-		Writer.AddString(ANSITEXTVIEW("level"), GetLevel(Verbosity));
-		AddMessage(Writer, V, Verbosity, Category, bShowCategory, bShowVerbosity);
-		if (bShowCategory || bShowVerbosity)
-		{
-			AddFormat(Writer, V, bShowCategory, bShowVerbosity);
-
-			Writer.BeginObject(ANSITEXTVIEW("properties"));
-			if (bShowCategory)
-			{
-				Writer.BeginObject(ANSITEXTVIEW("_channel"));
-				Writer.AddString(ANSITEXTVIEW("$type"), ANSITEXTVIEW("Channel"));
-				Writer.AddString(ANSITEXTVIEW("$text"), WriteToUtf8String<64>(Category));
-				Writer.EndObject();
-			}
-			if (bShowVerbosity)
-			{
-				Writer.BeginObject(ANSITEXTVIEW("_severity"));
-				Writer.AddString(ANSITEXTVIEW("$type"), ANSITEXTVIEW("Severity"));
-				Writer.AddString(ANSITEXTVIEW("$text"), ToString(Verbosity));
-				Writer.EndObject();
-			}
-			Writer.EndObject();
-		}
-		Writer.EndObject();
-
-		WriteAsJson(Writer);
-	}
-
-	FORCENOINLINE static void AddMessage(FCbWriter& Writer, const TCHAR* V, ELogVerbosity::Type Verbosity, const FName& Category, bool bShowCategory, bool bShowVerbosity)
-	{
-		TUtf8StringBuilder<512> Message;
-		if (bShowCategory)
-		{
-			Message << Category << ANSITEXTVIEW(": ");
-		}
-		if (bShowVerbosity)
-		{
-			Message << ToString(Verbosity) << ANSITEXTVIEW(": ");
-		}
-		Message << V;
-		Writer.AddString(ANSITEXTVIEW("message"), Message);
-	}
-
-	FORCENOINLINE static void AddFormat(FCbWriter& Writer, const TCHAR* Message, bool bShowCategory, bool bShowVerbosity)
-	{
-		TUtf8StringBuilder<512> Format;
-		if (bShowCategory)
-		{
-			Format << ANSITEXTVIEW("{_channel}: ");
-		}
-		if (bShowVerbosity)
-		{
-			Format << ANSITEXTVIEW("{_severity}: ");
-		}
-
-		// Escape {} in Message
-		for (constexpr FAsciiSet Brackets("{}");;)
-		{
-			const TCHAR* End = FAsciiSet::FindFirstOrEnd(Message, Brackets);
-			Format.Append(Message, UE_PTRDIFF_TO_INT32(End - Message));
-			if (!*End)
-			{
-				break;
-			}
-			Format.AppendChar(*End);
-			Format.AppendChar(*End);
-			Message = End + 1;
-		}
-
-		Writer.AddString(ANSITEXTVIEW("format"), Format);
-	}
-
-	FORCENOINLINE void SerializeRecordAsJson(const UE::FLogRecord& Record)
-	{
-		const bool bShowCategory = GPrintLogCategory && !Record.GetCategory().IsNone();
-		const bool bShowVerbosity = GPrintLogVerbosity && (Record.GetVerbosity() & ELogVerbosity::VerbosityMask) != ELogVerbosity::Log;
-
-		TCbWriter<1024> Writer;
-		Writer.BeginObject();
-		Writer.AddDateTime(ANSITEXTVIEW("time"), Record.GetTime().GetUtcTime());
-		Writer.AddString(ANSITEXTVIEW("level"), GetLevel(Record.GetVerbosity()));
-		AddMessage(Writer, Record, bShowCategory, bShowVerbosity);
-		if (bShowCategory || bShowVerbosity || Record.GetFields())
-		{
-			AddFormat(Writer, Record, bShowCategory, bShowVerbosity);
-
-			Writer.BeginObject(ANSITEXTVIEW("properties"));
-			if (bShowCategory)
-			{
-				Writer.BeginObject(ANSITEXTVIEW("_channel"));
-				Writer.AddString(ANSITEXTVIEW("$type"), ANSITEXTVIEW("Channel"));
-				Writer.AddString(ANSITEXTVIEW("$text"), WriteToUtf8String<64>(Record.GetCategory()));
-				Writer.EndObject();
-			}
-			if (bShowVerbosity)
-			{
-				Writer.BeginObject(ANSITEXTVIEW("_severity"));
-				Writer.AddString(ANSITEXTVIEW("$type"), ANSITEXTVIEW("Severity"));
-				Writer.AddString(ANSITEXTVIEW("$text"), ToString(Record.GetVerbosity()));
-				Writer.EndObject();
-			}
-			if (const TCHAR* TextNamespace = Record.GetTextNamespace())
-			{
-				Writer.AddString(ANSITEXTVIEW("_ns"), TextNamespace);
-			}
-			if (const TCHAR* TextKey = Record.GetTextKey())
-			{
-				Writer.AddString(ANSITEXTVIEW("_key"), TextKey);
-			}
-			for (const FCbField& Field : Record.GetFields())
-			{
-				Writer.AddField(Field.GetName(), Field);
-			}
-			Writer.EndObject();
-		}
-		Writer.EndObject();
-
-		WriteAsJson(Writer);
-	}
-
-	FORCENOINLINE static void AddMessage(FCbWriter& Writer, const UE::FLogRecord& Record, bool bShowCategory, bool bShowVerbosity)
-	{
-		TUtf8StringBuilder<512> Message;
-		if (bShowCategory)
-		{
-			Message << Record.GetCategory() << ANSITEXTVIEW(": ");
-		}
-		if (bShowVerbosity)
-		{
-			Message << ToString(Record.GetVerbosity()) << ANSITEXTVIEW(": ");
-		}
-		Record.FormatMessageTo(Message);
-		Writer.AddString(ANSITEXTVIEW("message"), Message);
-	}
-
-	FORCENOINLINE static void AddFormat(FCbWriter& Writer, const UE::FLogRecord& Record, bool bShowCategory, bool bShowVerbosity)
-	{
-		TUtf8StringBuilder<512> Format;
-		if (bShowCategory)
-		{
-			Format << ANSITEXTVIEW("{_channel}: ");
-		}
-		if (bShowVerbosity)
-		{
-			Format << ANSITEXTVIEW("{_severity}: ");
-		}
-		Format.Append(Record.GetFormat());
-		Writer.AddString(ANSITEXTVIEW("format"), Format);
-	}
-
-	void WriteAsJson(const FCbWriter& Writer)
-	{
-		TArray<uint8, TInlineAllocator64<512>> Buffer;
-		Buffer.AddUninitialized((int64)Writer.GetSaveSize());
-		FCbFieldView Object = Writer.Save(MakeMemoryView(Buffer));
-
-	#if PLATFORM_WINDOWS
-		if (bIsConsoleOutput)
-		{
-			WriteLine<WIDECHAR>(Object);
-			return;
-		}
-	#endif
-	#if PLATFORM_WINDOWS || PLATFORM_TCHAR_IS_UTF8CHAR || PLATFORM_TCHAR_IS_CHAR16
-		WriteLine<UTF8CHAR>(Object);
-	#else
-		WriteLine<WIDECHAR>(Object);
-	#endif
-	}
-
-	template <typename CharType>
-	FORCENOINLINE void WriteLine(const FCbFieldView& Field)
-	{
-		TStringBuilderWithBuffer<CharType, 512> Line;
-		CompactBinaryToCompactJson(Field, Line);
-		Line.AppendChar('\n');
-		WriteLine(Line);
-	}
-
-	void WriteLine(FUtf8StringBuilderBase& Line) const
-	{
-		printf("%s", (const char*)*Line);
-		fflush(stdout);
-	}
-
-#if PLATFORM_WINDOWS || !(PLATFORM_TCHAR_IS_UTF8CHAR || PLATFORM_TCHAR_IS_CHAR16)
-	void WriteLine(FWideStringBuilderBase& Line) const
-	{
-	#if PLATFORM_WINDOWS
-		if (bIsConsoleOutput)
-		{
-			WriteConsoleW(GetStdHandle(STD_OUTPUT_HANDLE), *Line, Line.Len(), nullptr, nullptr);
-			return;
-		}
-	#endif
-
-	#if PLATFORM_USE_LS_SPEC_FOR_WIDECHAR
-		// printf prints wchar_t strings just fine with %ls, while mixing printf()/wprintf() is not recommended (see https://stackoverflow.com/questions/8681623/printf-and-wprintf-in-single-c-code)
-		printf("%ls", *Line);
-	#else
-		wprintf(TEXT("%s"), *Line);
-	#endif
-		fflush(stdout);
-	}
-#endif
-
-	static FAnsiStringView GetLevel(ELogVerbosity::Type Verbosity)
-	{
-		switch (Verbosity & ELogVerbosity::VerbosityMask)
-		{
-		case ELogVerbosity::Fatal:
-			return ANSITEXTVIEW("Critical");
-		case ELogVerbosity::Error:
-			return ANSITEXTVIEW("Error");
-		case ELogVerbosity::Warning:
-			return ANSITEXTVIEW("Warning");
-		case ELogVerbosity::Display:
-		case ELogVerbosity::Log:
-		default:
-			return ANSITEXTVIEW("Information");
-		case ELogVerbosity::Verbose:
-		case ELogVerbosity::VeryVerbose:
-			return ANSITEXTVIEW("Debug");
-		}
-	}
-
-private:
-	ELogVerbosity::Type AllowedLogVerbosity = ELogVerbosity::Display;
-	bool				bIsConsoleOutput = false;
-	bool				bIsJsonOutput = false;
-
-#if PLATFORM_WINDOWS
-	static bool IsStdoutAttachedToConsole()
-	{
-		HANDLE StdoutHandle = GetStdHandle(STD_OUTPUT_HANDLE);
-		if (StdoutHandle != INVALID_HANDLE_VALUE)
-		{
-			DWORD FileType = GetFileType(StdoutHandle);
-			if (FileType == FILE_TYPE_CHAR)
-			{
-				return true;
-			}
-		}
-
-		return false;
-	}
-#endif
-};
-
-
 // Exits the game/editor if any of the specified phrases appears in the log output
 class FOutputDeviceTestExit : public FOutputDevice
 {
