Always initialize system textures after RHI init
* Doing the initialization in FEngineLoop::PostInitRHI() ensure GMaxRHIShaderPlatform is correct
* Will prevent the need to add this initialization manually in multiple systems
* Tested in commandlets, editor and a cooked build

diff --git a/Engine/Source/Runtime/Launch/Private/LaunchEngineLoop.cpp b/Engine/Source/Runtime/Launch/Private/LaunchEngineLoop.cpp
--- a/Engine/Source/Runtime/Launch/Private/LaunchEngineLoop.cpp
+++ b/Engine/Source/Runtime/Launch/Private/LaunchEngineLoop.cpp
@@ -6677,6 +6677,12 @@ void FEngineLoop::PostInitRHI()
 
 		UE::StereoRenderUtils::LogISRInit(Aspects);
 		UE::StereoRenderUtils::VerifyISRConfig(Aspects, ShaderPlatform);
+		
+		// Initialize system textures
+		ENQUEUE_RENDER_COMMAND(InitializeSystemTextures)([](FRHICommandListImmediate& RHICmdList)
+		{
+			GetRendererModule().InitializeSystemTextures(RHICmdList);
+		});
 	}
 
 #endif
