[WinArm64] fix for "editor not built" startup error message due to not finding the correct .target file for arm64ec

#rnx

diff --git a/Engine/Source/Runtime/Core/Private/GenericPlatform/GenericPlatformProcess.cpp b/Engine/Source/Runtime/Core/Private/GenericPlatform/GenericPlatformProcess.cpp
--- a/Engine/Source/Runtime/Core/Private/GenericPlatform/GenericPlatformProcess.cpp
+++ b/Engine/Source/Runtime/Core/Private/GenericPlatform/GenericPlatformProcess.cpp
@@ -211,6 +211,11 @@ FString FGenericPlatformProcess::GenerateApplicationPath( const FString& AppName
 	return FString();
 }
 
+const TCHAR* FGenericPlatformProcess::GetArchitectureSuffix()
+{
+	return nullptr;
+}
+
 const TCHAR* FGenericPlatformProcess::GetModulePrefix()
 {
 	return TEXT("");
diff --git a/Engine/Source/Runtime/Core/Private/Windows/WindowsPlatformProcess.cpp b/Engine/Source/Runtime/Core/Private/Windows/WindowsPlatformProcess.cpp
--- a/Engine/Source/Runtime/Core/Private/Windows/WindowsPlatformProcess.cpp
+++ b/Engine/Source/Runtime/Core/Private/Windows/WindowsPlatformProcess.cpp
@@ -269,6 +269,17 @@ FString FWindowsPlatformProcess::GenerateApplicationPath( const FString& AppName
 	return ExecutablePath;
 }
 
+const TCHAR* FWindowsPlatformProcess::GetArchitectureSuffix()
+{
+#if defined(_M_ARM64EC)
+	return TEXT("arm64ec");
+#elif defined(_M_ARM64)
+	return TEXT("arm64");
+#else
+	return nullptr;
+#endif
+}
+
 void* FWindowsPlatformProcess::GetDllExport( void* DllHandle, const TCHAR* ProcName )
 {
 	check(DllHandle);
diff --git a/Engine/Source/Runtime/Core/Public/GenericPlatform/GenericPlatformProcess.h b/Engine/Source/Runtime/Core/Public/GenericPlatform/GenericPlatformProcess.h
--- a/Engine/Source/Runtime/Core/Public/GenericPlatform/GenericPlatformProcess.h
+++ b/Engine/Source/Runtime/Core/Public/GenericPlatform/GenericPlatformProcess.h
@@ -378,6 +378,15 @@ struct FGenericPlatformProcess
 	 */
 	static CORE_API FString GenerateApplicationPath( const FString& AppName, EBuildConfiguration BuildConfiguration);
 
+	/**
+	 * Get the architecture suffix, if any
+	 * 
+	 * Relates to UnrealArchitectureConfig.RequiresArchitectureFilenames in UnrealBuildTool
+	 * 
+	 * @return The platform's architecture suffix, or null if it isn't required
+	 */
+	static CORE_API const TCHAR* GetArchitectureSuffix();
+
 	/**
 	 * Return the prefix of dynamic library (e.g. lib)
 	 *
diff --git a/Engine/Source/Runtime/Core/Public/Windows/WindowsPlatformProcess.h b/Engine/Source/Runtime/Core/Public/Windows/WindowsPlatformProcess.h
--- a/Engine/Source/Runtime/Core/Public/Windows/WindowsPlatformProcess.h
+++ b/Engine/Source/Runtime/Core/Public/Windows/WindowsPlatformProcess.h
@@ -159,6 +159,7 @@ public:
 	static CORE_API const TCHAR* ExecutablePath();
 	static CORE_API const TCHAR* ExecutableName(bool bRemoveExtension = true);
 	static CORE_API FString GenerateApplicationPath( const FString& AppName, EBuildConfiguration BuildConfiguration);
+	static CORE_API const TCHAR* GetArchitectureSuffix();
 	static CORE_API const TCHAR* GetModuleExtension();
 	static CORE_API const TCHAR* GetBinariesSubdirectory();
 	static CORE_API const FString GetModulesDirectory();
diff --git a/Engine/Source/Runtime/Launch/Private/LaunchEngineLoop.cpp b/Engine/Source/Runtime/Launch/Private/LaunchEngineLoop.cpp
--- a/Engine/Source/Runtime/Launch/Private/LaunchEngineLoop.cpp
+++ b/Engine/Source/Runtime/Launch/Private/LaunchEngineLoop.cpp
@@ -6234,14 +6234,16 @@ bool FEngineLoop::AppInit( )
 		{
 			if (Target.Type == EBuildTargetType::Editor)
 			{
+				const TCHAR* ArchitectureSuffix = FPlatformProcess::GetArchitectureSuffix();
+
 				FString ScratchEditorTargetFileName;
 				if (FPaths::IsUnderDirectory(Target.Path, FPlatformMisc::ProjectDir()))
 				{
-					ScratchEditorTargetFileName = FTargetReceipt::GetDefaultPath(FPlatformMisc::ProjectDir(), *Target.Name, FPlatformProcess::GetBinariesSubdirectory(), FApp::GetBuildConfiguration(), nullptr);
+					ScratchEditorTargetFileName = FTargetReceipt::GetDefaultPath(FPlatformMisc::ProjectDir(), *Target.Name, FPlatformProcess::GetBinariesSubdirectory(), FApp::GetBuildConfiguration(), ArchitectureSuffix);
 				}
 				else if (FPaths::IsUnderDirectory(Target.Path, FPaths::EngineDir()))
 				{
-					ScratchEditorTargetFileName = FTargetReceipt::GetDefaultPath(*FPaths::EngineDir(), *Target.Name, FPlatformProcess::GetBinariesSubdirectory(), FApp::GetBuildConfiguration(), nullptr);
+					ScratchEditorTargetFileName = FTargetReceipt::GetDefaultPath(*FPaths::EngineDir(), *Target.Name, FPlatformProcess::GetBinariesSubdirectory(), FApp::GetBuildConfiguration(), ArchitectureSuffix);
 				}
 
 				if (DoesCurrentExecutableMatch(ScratchEditorTargetFileName))
