[ASIS] Implementation of touch support for extra view

#rnx

diff --git a/Engine/Source/Runtime/Launch/Private/Android/LaunchAndroid.cpp b/Engine/Source/Runtime/Launch/Private/Android/LaunchAndroid.cpp
--- a/Engine/Source/Runtime/Launch/Private/Android/LaunchAndroid.cpp
+++ b/Engine/Source/Runtime/Launch/Private/Android/LaunchAndroid.cpp
@@ -2174,78 +2174,6 @@ JNI_METHOD void Java_com_epicgames_makeaar_GameActivityForMakeAAR_nativeMain(JNI
 	STANDALONE_DEBUG_LOGf(LogAndroid, TEXT("Created main thread"));
 }
 
-JNI_METHOD jint Java_com_epicgames_makeaar_GameActivityForMakeAAR_nativeInputTouch(JNIEnv* jenv, jobject thiz, jint device, jint inType, jint inPointerId, jint inX, jint inY)
-{
-	TArray<TouchInput> TouchesArray;
-
-	TouchType type = TouchEnded;
-	switch (inType)
-	{
-		case 0: type = TouchBegan; break;
-		case 1: type = TouchMoved; break;
-		case 2: type = TouchEnded; break;
-	}
-
-	ANativeWindow* Window = (ANativeWindow*)FAndroidWindow::GetHardwareWindow_EventThread();
-	if (!Window)
-	{
-		return 0;
-	}
-
-	int32_t Width = 0;
-	int32_t Height = 0;
-
-	if (Window)
-	{
-		if (GAndroidWindowOverride != nullptr)
-		{
-			Width = ANativeWindow_getWidth((ANativeWindow*)GAndroidWindowOverride);
-			Height = ANativeWindow_getHeight((ANativeWindow*)GAndroidWindowOverride);
-		}
-		else
-		{
-			// we are on the event thread. true here indicates we will retrieve dimensions from the current window.
-			FAndroidWindow::CalculateSurfaceSize(Width, Height, true);
-		}
-	}
-
-	// make sure OpenGL context created before accepting touch events.. FAndroidWindow::GetScreenRect() may try to create it early from wrong thread if this is the first call
-	if (!GAndroidGPUInfoReady)
-	{
-		return 1;
-	}
-	FPlatformRect ScreenRect = FAndroidWindow::GetScreenRect(true);
-	int32_t ScreenWidth = ScreenRect.Right - ScreenRect.Left;
-	int32_t ScreenHeight = ScreenRect.Bottom - ScreenRect.Top;
-
-	int32 eventX = inX - ScreenRect.Left;
-	int32 eventY = inY - ScreenRect.Top;
-	if (eventX < 0 || eventY < 0 || eventX >= ScreenWidth || eventY >= ScreenHeight)
-	{
-		return 0;
-	}
-
-	int pointerId = inPointerId;
-	float x = FMath::Max<float>(FMath::Min<float>((float)eventX / Width, 1.f), 0.f) * (ScreenWidth - 1);
-	float y = FMath::Max<float>(FMath::Min<float>((float)eventY / Height, 1.f), 0.f) * (ScreenHeight - 1);
-
-	//STANDALONE_DEBUG_LOGf(LogAndroid, TEXT("TOUCH: at (%d,%d) windowPos = (%d,%d), windowSize = (%d,%d), screenSize = (%d,%d), final pos: (%f,%f)"),
-	//	eventX, eventY, ScreenRect.Left, ScreenRect.Top, ScreenWidth, ScreenHeight, Width, Height, x, y);
-	UE_LOG(LogAndroid, Verbose, TEXT("Received targeted motion event from pointer %u (id %d) action %d: (%.2f, %.2f)"), 0, pointerId, type, x, y);
-
-	TouchInput TouchMessage;
-	TouchMessage.DeviceId = device;
-	TouchMessage.Handle = pointerId;
-	TouchMessage.Type = type;
-	TouchMessage.Position = FVector2D(x, y);
-	TouchMessage.LastPosition = FVector2D(x, y);
-	TouchesArray.Add(TouchMessage);
-
-	FAndroidInputInterface::QueueTouchInput(TouchesArray);
-
-	return 1;
-}
-
 static bool MountPak(const FString& PakFilePath, int32 PakOrder, const FString& PakMountPoint)
 {
 	// Mount all pak files found in this content
