Fix shutdown crashes by passing flag to stop IOStore unloads from calling dangerous callbacks

diff --git a/Engine/Source/Runtime/Launch/Private/LaunchEngineLoop.cpp b/Engine/Source/Runtime/Launch/Private/LaunchEngineLoop.cpp
--- a/Engine/Source/Runtime/Launch/Private/LaunchEngineLoop.cpp
+++ b/Engine/Source/Runtime/Launch/Private/LaunchEngineLoop.cpp
@@ -4876,9 +4876,9 @@ void FEngineLoop::Exit()
 #if USE_IO_DISPATCHER
 	FIoDispatcher::Shutdown();
 #if UE_WITH_IOSTOREONDEMAND
-	FModuleManager::Get().UnloadModule("IoStoreOnDemand");
+	FModuleManager::Get().UnloadModule("IoStoreOnDemand", true);
 #endif
-	FModuleManager::Get().UnloadModule("IoStoreOnDemandCore");
+	FModuleManager::Get().UnloadModule("IoStoreOnDemandCore", true);
 #endif
 
 #if WITH_EDITOR
