Unify the module loading rules so that PostSplashScreen modules load in all cases

diff --git a/Engine/Source/Runtime/Launch/Private/LaunchEngineLoop.cpp b/Engine/Source/Runtime/Launch/Private/LaunchEngineLoop.cpp
--- a/Engine/Source/Runtime/Launch/Private/LaunchEngineLoop.cpp
+++ b/Engine/Source/Runtime/Launch/Private/LaunchEngineLoop.cpp
@@ -3177,6 +3177,17 @@ int32 FEngineLoop::PreInitPreStartupScreen(const TCHAR* CmdLine)
 
 				PreInitContext.SlateRenderer = SlateRenderer;
 			}
+			else
+			{
+				{
+					SCOPED_BOOT_TIMING("LoadModulesForProject(ELoadingPhase::PostSplashScreen)");
+					// Load up all modules that need to hook into the custom splash screen
+					if (!IProjectManager::Get().LoadModulesForProject(ELoadingPhase::PostSplashScreen) || !IPluginManager::Get().LoadModulesForEnabledPlugins(ELoadingPhase::PostSplashScreen))
+					{
+						return 1;
+					}
+				}
+			}
 #endif // !UE_SERVER
 		}
 	}
