[ASIS] Rename MakeAAR to ASIS

#rnx

diff --git a/Engine/Source/Programs/UnrealBuildTool/Platform/Android/UEBuildAndroid.cs b/Engine/Source/Programs/UnrealBuildTool/Platform/Android/UEBuildAndroid.cs
--- a/Engine/Source/Programs/UnrealBuildTool/Platform/Android/UEBuildAndroid.cs
+++ b/Engine/Source/Programs/UnrealBuildTool/Platform/Android/UEBuildAndroid.cs
@@ -443,7 +443,7 @@ namespace UnrealBuildTool
 		{
 		}
 
-		public static bool IsMakeAAREnabled(ReadOnlyTargetRules Target)
+		public static bool IsAsisEnabled(ReadOnlyTargetRules Target)
 		{
 			return Target.AndroidPlatform.bEnableASISPlugin;
 		}
diff --git a/Engine/Source/Runtime/ApplicationCore/Private/Android/AndroidWindow.cpp b/Engine/Source/Runtime/ApplicationCore/Private/Android/AndroidWindow.cpp
--- a/Engine/Source/Runtime/ApplicationCore/Private/Android/AndroidWindow.cpp
+++ b/Engine/Source/Runtime/ApplicationCore/Private/Android/AndroidWindow.cpp
@@ -209,33 +209,6 @@ JNI_METHOD void Java_com_epicgames_unreal_GameActivity_nativeSetSafezoneInfo(JNI
 #endif
 	UE_LOG(LogAndroid, Log, TEXT("nativeSetSafezoneInfo bIsPortrait=%d, left=%f, top=%f, right=%f, bottom=%f"), bIsPortrait ? 1 : 0, left, top, right, bottom);
 }
-
-#if USE_ANDROID_STANDALONE
-
-JNI_METHOD void Java_com_epicgames_makeaar_GameActivityForMakeAAR_nativeSetSurfaceOverride(JNIEnv* jenv, jobject thiz, jobject surface, jint x, jint y)
-{
-	ANativeWindow* prev = (ANativeWindow*)GAndroidWindowOverride;
-	if (surface != 0)
-	{
-		GAndroidWindowOverride = (ANativeWindow*)ANativeWindow_fromSurface(jenv, surface);
-		UE_LOG(LogAndroid, Log, TEXT("nativeSetSurfaceOverride applied: prev to new %p -> %p, pos(%d, %d)"), prev, GAndroidWindowOverride, x, y);
-
-	}
-	else
-	{
-		GAndroidWindowOverride = nullptr;
-		
-		STANDALONE_DEBUG_LOG(TEXT("nativeSetSurfaceOverride(makeaar) setting to null"));
-	}
-
-	if (prev != nullptr)
-	{
-		ANativeWindow_release(prev);
-	}
-}
-
-#endif // USE_ANDROID_STANDALONE
-
 #endif
 
 bool FAndroidWindow::bAreCachedNativeDimensionsValid = false;
diff --git a/Engine/Source/Runtime/Core/Private/Android/AndroidPlatformFile.cpp b/Engine/Source/Runtime/Core/Private/Android/AndroidPlatformFile.cpp
--- a/Engine/Source/Runtime/Core/Private/Android/AndroidPlatformFile.cpp
+++ b/Engine/Source/Runtime/Core/Private/Android/AndroidPlatformFile.cpp
@@ -140,7 +140,7 @@ JNI_METHOD void Java_com_epicgames_unreal_GameActivity_nativeSetObbInfo(JNIEnv*
 }
 
 //This function is declared in the Java-defined class, GameActivity.java: "public native String nativeGetObbComment();"
-JNI_METHOD jstring Java_com_epicgames_makeaar_GameActivityForMakeAAR_nativeGetObbComment(JNIEnv* jenv, jobject thiz)
+JNI_METHOD jstring Java_com_epicgames_asis_AsisGameActivity_nativeGetObbComment(JNIEnv* jenv, jobject thiz)
 {
 	return jenv->NewStringUTF(TCHAR_TO_UTF8(*GLastOBBComment));
 }
diff --git a/Engine/Source/Runtime/Core/Public/Android/AndroidPlatform.h b/Engine/Source/Runtime/Core/Public/Android/AndroidPlatform.h
--- a/Engine/Source/Runtime/Core/Public/Android/AndroidPlatform.h
+++ b/Engine/Source/Runtime/Core/Public/Android/AndroidPlatform.h
@@ -119,7 +119,7 @@ typedef FAndroidTypes FPlatformTypes;
 #define ANDROID_GAMEACTIVITY_BASE_CLASSPATH			"com/epicgames/unreal/GameActivity"
 
 #if USE_ANDROID_STANDALONE
-#	define ANDROID_GAMEACTIVITY_CLASSPATH	"com/epicgames/makeaar/GameActivityForMakeAAR"
+#	define ANDROID_GAMEACTIVITY_CLASSPATH	"com/epicgames/asis/AsisGameActivity"
 #else
 #	define ANDROID_GAMEACTIVITY_CLASSPATH	"com/epicgames/unreal/GameActivity"
 #endif
diff --git a/Engine/Source/Runtime/Launch/Private/Android/AndroidJNI.cpp b/Engine/Source/Runtime/Launch/Private/Android/AndroidJNI.cpp
--- a/Engine/Source/Runtime/Launch/Private/Android/AndroidJNI.cpp
+++ b/Engine/Source/Runtime/Launch/Private/Android/AndroidJNI.cpp
@@ -2393,12 +2393,12 @@ JNI_METHOD void Java_com_epicgames_unreal_NativeCalls_UELogVerbose(JNIEnv* jenv,
 	jenv->ReleaseStringUTFChars(InString, chars);
 }
 #if USE_ANDROID_STANDALONE
-JNI_METHOD void Java_com_epicgames_makeaar_GameActivityForMakeAAR_nativeSetGlobalActivity(JNIEnv* jenv, jobject thiz, jboolean bUseExternalFilesDir, jboolean bPublicLogFiles, jstring internalFilePath, jstring externalFilePath, jboolean bOBBinAPK, jstring APKFilename /*, jobject googleServices*/)
+JNI_METHOD void Java_com_epicgames_asis_AsisGameActivity_nativeSetGlobalActivity(JNIEnv* jenv, jobject thiz, jboolean bUseExternalFilesDir, jboolean bPublicLogFiles, jstring internalFilePath, jstring externalFilePath, jboolean bOBBinAPK, jstring APKFilename /*, jobject googleServices*/)
 {
-	STANDALONE_DEBUG_LOG(TEXT("nativeSetGlobalActivity: Entering makeaar nativeSetGlobalActivity, GameActivityThis=%p\n"), FJavaWrapper::GameActivityThis);
+	STANDALONE_DEBUG_LOG(TEXT("nativeSetGlobalActivity: Entering asis nativeSetGlobalActivity, GameActivityThis=%p\n"), FJavaWrapper::GameActivityThis);
 	if (FJavaWrapper::GameActivityThis != nullptr)
 	{
-		STANDALONE_DEBUG_LOG(TEXT("nativeSetGlobalActivity(makeaar): Error GameActivityThis is already set GameActivityThis=%p\n"), FJavaWrapper::GameActivityThis);
+		STANDALONE_DEBUG_LOG(TEXT("nativeSetGlobalActivity(asis): Error GameActivityThis is already set GameActivityThis=%p\n"), FJavaWrapper::GameActivityThis);
 		//jenv->DeleteGlobalRef(FJavaWrapper::GameActivityThis);
 		GGameActivityThis = FJavaWrapper::GameActivityThis = nullptr;
 	}
@@ -2410,10 +2410,10 @@ JNI_METHOD void Java_com_epicgames_makeaar_GameActivityForMakeAAR_nativeSetGloba
 			FPlatformMisc::LowLevelOutputDebugString(TEXT("Error setting the global GameActivity activity"));
 			check(false);
 		}
-		STANDALONE_DEBUG_LOG(TEXT("nativeSetGlobalActivity(makeaar): jenv=%p, set GameActivityThis=%p, bUseExternalFilesDir=%d, bOBBinAPK=%d\n"), jenv, FJavaWrapper::GameActivityThis, bUseExternalFilesDir, bOBBinAPK);
+		STANDALONE_DEBUG_LOG(TEXT("nativeSetGlobalActivity(asis): jenv=%p, set GameActivityThis=%p, bUseExternalFilesDir=%d, bOBBinAPK=%d\n"), jenv, FJavaWrapper::GameActivityThis, bUseExternalFilesDir, bOBBinAPK);
 		// This call is only to set the correct GameActivityThis
 		FAndroidApplication::InitializeJavaEnv(GJavaVM, JNI_CURRENT_VERSION, FJavaWrapper::GameActivityThis);
-		// Rescan methods since we are switching to the makeaar version of GameActivity
+		// Rescan methods since we are switching to the asis version of GameActivity
 		FJavaWrapper::FindClassesAndMethods(jenv);
 		// @todo split GooglePlay, this needs to be passed in to this function
 		FJavaWrapper::GoogleServicesThis = FJavaWrapper::GameActivityThis;
@@ -2433,10 +2433,10 @@ JNI_METHOD void Java_com_epicgames_makeaar_GameActivityForMakeAAR_nativeSetGloba
 #else
 			GFilePathBase = GExternalFilePath;
 #endif
-			FPlatformMisc::LowLevelOutputDebugStringf(TEXT("nativeSetGlobalActivity(makeaar): GFilePathBase Path override to'%s'\n"), *GFilePathBase);
+			FPlatformMisc::LowLevelOutputDebugStringf(TEXT("nativeSetGlobalActivity(asis): GFilePathBase Path override to'%s'\n"), *GFilePathBase);
 		}
-		FPlatformMisc::LowLevelOutputDebugStringf(TEXT("nativeSetGlobalActivity(makeaar): InternalFilePath found as '%s'\n"), *GInternalFilePath);
-		FPlatformMisc::LowLevelOutputDebugStringf(TEXT("nativeSetGlobalActivity(makeaar): ExternalFilePath found as '%s'\n"), *GExternalFilePath);
+		FPlatformMisc::LowLevelOutputDebugStringf(TEXT("nativeSetGlobalActivity(asis): InternalFilePath found as '%s'\n"), *GInternalFilePath);
+		FPlatformMisc::LowLevelOutputDebugStringf(TEXT("nativeSetGlobalActivity(asis): ExternalFilePath found as '%s'\n"), *GExternalFilePath);
 	}
 	// DO NOT CALL, we have already handled logic above...
 	// Java_com_epicgames_unreal_GameActivity_nativeSetGlobalActivity(jenv, thiz, bUseExternalFilesDir, bPublicLogFiles, internalFilePath, externalFilePath, bOBBinAPK, APKFilename);
diff --git a/Engine/Source/Runtime/Launch/Private/Android/LaunchAndroid.cpp b/Engine/Source/Runtime/Launch/Private/Android/LaunchAndroid.cpp
--- a/Engine/Source/Runtime/Launch/Private/Android/LaunchAndroid.cpp
+++ b/Engine/Source/Runtime/Launch/Private/Android/LaunchAndroid.cpp
@@ -2148,7 +2148,7 @@ void AndroidThunkCpp_Engine_SendEvent(int32 event, const FString& param1, int32
 {
 	if (JNIEnv* Env = FAndroidApplication::GetJavaEnv())
 	{
-		static jclass EngineClass = FAndroidApplication::FindJavaClassGlobalRef("com/epicgames/makeaar/Engine");
+		static jclass EngineClass = FAndroidApplication::FindJavaClassGlobalRef("com/epicgames/asis/Engine");
 		static jmethodID SendEventFunc = FJavaWrapper::FindStaticMethod(Env, EngineClass, "AndroidThunkJava_Engine_ReceiveEvent", "(ILjava/lang/String;IIF)V", false);
 
 		if (SendEventFunc != nullptr)
@@ -2222,20 +2222,20 @@ static void PostLoadMapEvent(UWorld* World)
 	}
 }
 
-JNI_METHOD void Java_com_epicgames_makeaar_GameActivityForMakeAAR_nativeSetCommandline(JNIEnv* jenv, jobject thiz, jstring commandline)
+JNI_METHOD void Java_com_epicgames_asis_AsisGameActivity_nativeSetCommandline(JNIEnv* jenv, jobject thiz, jstring commandline)
 {
 	auto UECommandLine = FJavaHelper::FStringFromParam(jenv, commandline);
 	GAndroidCommandLine = UECommandLine;
 }
 
-JNI_METHOD void Java_com_epicgames_makeaar_GameActivityForMakeAAR_nativeMain(JNIEnv* jenv, jobject thiz, jstring projectModule)
+JNI_METHOD void Java_com_epicgames_asis_AsisGameActivity_nativeMain(JNIEnv* jenv, jobject thiz, jstring projectModule)
 {
 	BootTimingPoint("native_main");
 	GIsGameAgnosticExe = true;
 
 	FString projectModuleName = FJavaHelper::FStringFromParam(jenv, projectModule);
 
-	STANDALONE_DEBUG_LOGf(LogAndroid, TEXT("Java_com_epicgames_makeaar_GameActivityForMakeAAR_nativeMain : use current!, requesting ProjectModule: %s"), *projectModuleName);
+	STANDALONE_DEBUG_LOGf(LogAndroid, TEXT("Java_com_epicgames_asis_AsisGameActivity_nativeMain : use current!, requesting ProjectModule: %s"), *projectModuleName);
 
 	if (EventThreadID == 0)
 	{
@@ -2284,7 +2284,7 @@ static bool MountPak(const FString& PakFilePath, int32 PakOrder, const FString&
 	return PakFileMgr->Mount(*PakFilePath, PakOrder);
 }
 
-JNI_METHOD bool Java_com_epicgames_makeaar_Engine_nativeMountPak(JNIEnv* jenv, jobject thiz, jstring pakFile, jint order, jstring mountPoint)
+JNI_METHOD bool Java_com_epicgames_asis_Engine_nativeMountPak(JNIEnv* jenv, jobject thiz, jstring pakFile, jint order, jstring mountPoint)
 {
 	FString PakFilePath = FJavaHelper::FStringFromParam(jenv, pakFile);
 	uint32 PakOrder = (uint32)FMath::Max(0, order);
@@ -2307,9 +2307,9 @@ JNI_METHOD bool Java_com_epicgames_makeaar_Engine_nativeMountPak(JNIEnv* jenv, j
 	}
 }
 
-JNI_METHOD void Java_com_epicgames_makeaar_GameActivityForMakeAAR_nativeAppCommand(JNIEnv* jenv, jobject thiz, jint cmd)
+JNI_METHOD void Java_com_epicgames_asis_AsisGameActivity_nativeAppCommand(JNIEnv* jenv, jobject thiz, jint cmd)
 {
-	STANDALONE_DEBUG_LOGf(LogAndroid, TEXT("GameActivityForMakeAAR -> nativeAppCommand[%u, %s] -> EventThreadID=%d, GGameThreadId=%d, GResumeMainInit=%d, GEventHandlerInitialized=%d, FTaskGraphInterface::IsRunning()=%d, FPlatformTLS::GetCurrentThreadId()=%d"), 
+	STANDALONE_DEBUG_LOGf(LogAndroid, TEXT("AsisGameActivity -> nativeAppCommand[%u, %s] -> EventThreadID=%d, GGameThreadId=%d, GResumeMainInit=%d, GEventHandlerInitialized=%d, FTaskGraphInterface::IsRunning()=%d, FPlatformTLS::GetCurrentThreadId()=%d"), 
 		cmd,
 		get_APP_CMD_String(cmd),
 		EventThreadID, 
