CPU Utilization:

Standardize so that the utilization is always in percent not just a fraction
- Windows was the only outlier

diff --git a/Engine/Source/Runtime/Core/Private/Windows/WindowsPlatformProcess.cpp b/Engine/Source/Runtime/Core/Private/Windows/WindowsPlatformProcess.cpp
--- a/Engine/Source/Runtime/Core/Private/Windows/WindowsPlatformProcess.cpp
+++ b/Engine/Source/Runtime/Core/Private/Windows/WindowsPlatformProcess.cpp
@@ -832,8 +832,8 @@ bool FWindowsPlatformProcess::GetPerFrameProcessorUsage(uint32 ProcessId, float&
 
 	if (bSuccess)
 	{
-		ProcessUsageFraction = LastProcessTime;
-		IdleUsageFraction = LastIdleTime;
+		ProcessUsageFraction = LastProcessTime * 100.0f;
+		IdleUsageFraction = LastIdleTime * 100.0f;
 	}
 	else
 	{
diff --git a/Engine/Source/Runtime/Launch/Private/LaunchEngineLoop.cpp b/Engine/Source/Runtime/Launch/Private/LaunchEngineLoop.cpp
--- a/Engine/Source/Runtime/Launch/Private/LaunchEngineLoop.cpp
+++ b/Engine/Source/Runtime/Launch/Private/LaunchEngineLoop.cpp
@@ -1342,17 +1342,17 @@ static void UpdateCoreCsvStats_BeginFrame()
 		{
 #if !UE_BUILD_SHIPPING
 			const uint32 ProcessId = (uint32)FPlatformProcess::GetCurrentProcessId();
-		    float ProcessUsageFraction = 0.f, IdleUsageFraction = 0.f;
-			if (FPlatformProcess::GetPerFrameProcessorUsage(ProcessId, ProcessUsageFraction, IdleUsageFraction))
+		    float ProcessUsagePercent = 0.f, IdleUsagePercent = 0.f;
+			if (FPlatformProcess::GetPerFrameProcessorUsage(ProcessId, ProcessUsagePercent, IdleUsagePercent))
 			{
 				// This occasionally happens on Linux and I am not entirely sure why.
 				// If we have zero idle time and zero utilization assume we have full idle time
-				if (ProcessUsageFraction == 0.0f && IdleUsageFraction == 0.0f)
+				if (ProcessUsagePercent == 0.0f && IdleUsagePercent == 0.0f)
 				{
-					IdleUsageFraction = 1.0f;
+					IdleUsagePercent = 1.0f;
 				}
-				CSV_CUSTOM_STAT_GLOBAL(CPUUsage_Process, ProcessUsageFraction, ECsvCustomStatOp::Set);
-				CSV_CUSTOM_STAT_GLOBAL(CPUUsage_Idle, IdleUsageFraction, ECsvCustomStatOp::Set);
+				CSV_CUSTOM_STAT_GLOBAL(CPUUsage_Process, ProcessUsagePercent, ECsvCustomStatOp::Set);
+				CSV_CUSTOM_STAT_GLOBAL(CPUUsage_Idle, IdleUsagePercent, ECsvCustomStatOp::Set);
 			}
 #endif
 			
