SwiftUI splash dialog improvements
-Fix the problem where audio goes away when you close the splash dialog.  The issue is that the audio is associated with the window such that when the window closes audio is killed.  We can avoid that simply by dismissing the window before we start up audio, and really the dialog should auto-dismiss anyway.
-Added a countdown display to the auto-start, mainly because it is less confusing.
-Note that in this implementation the dialog goes away and then it takes unreal 2-3 seconds to present an immersive frame, so we get a gap where nothing is displayed.  I will probably want to figure out a way to dismiss the window later, or pop up a new window, or something like that.

diff --git a/Engine/Source/Runtime/Launch/Private/Apple/SwiftMain.swift b/Engine/Source/Runtime/Launch/Private/Apple/SwiftMain.swift
--- a/Engine/Source/Runtime/Launch/Private/Apple/SwiftMain.swift
+++ b/Engine/Source/Runtime/Launch/Private/Apple/SwiftMain.swift
@@ -8,42 +8,51 @@ struct SwiftUIView: View {
 	var onClick: () -> Void = {}
 
 	@Environment(\.openImmersiveSpace) var openImmersiveSpace
+	@Environment(\.dismissWindow) private var dismissWindow
 
 	@State private var buttonDisabled = false
-
+	@State private var timerCounter = 3
+	@State private var timerText = "3"
+	let timer = Timer.publish(every: 1, on: .main, in: .common).autoconnect()
+	
 	func Open()
 	{
+		print("SwiftUIView Open()")
+
+		buttonDisabled = true
+		timer.upstream.connect().cancel()
+
 		Task
 		{
-			buttonDisabled = true
 			await openImmersiveSpace(id: "ImmersiveSpace")
+			
+			dismissWindow(id: "SwiftUIView")
 		}
 	}
 	
-	let timer = Timer.publish(every: 3, on: .main, in: .common).autoconnect()
-	
 	var body: some View {
-		
-	VStack(spacing: 8) {
-	  Text("SwiftUI in Unreal!")
-		.font(.title)
-		.bold()
-	  Button(action: {
-		Open()
-//		onClick()
-	  }
-	  , label: {
-		  Text("Test Button")
-		})
-		.disabled(buttonDisabled)
-
-	}
-	.onReceive(timer) { t2 in
-		print("timer!")
-		timer.upstream.connect().cancel()
-		Open()
+		VStack(spacing: 8) {
+			Text("Unreal SwiftMain Splash Screen!")
+				.font(.title)
+				.bold()
+			Button(action: {
+				Open()
+			}
+				   , label: {
+				Text("Go Immersive!")
+			})
+			.disabled(buttonDisabled)
+			Text("Automatically go in \(timerText)")
+				.id("Countdown")
+		}
+		.onReceive(timer) { t2 in
+			timerCounter -= 1	
+			timerText = "\(timerCounter)"
+			if timerCounter == 0 {
+				Open()
+			}
+		}
 	}
-  }
 }
 
 struct UEContentConfiguration: CompositorLayerConfiguration {
@@ -112,7 +121,7 @@ struct UESwiftApp: App {
 	
 	var body: some Scene
 	{
-		WindowGroup
+		WindowGroup(id: "SwiftUIView")
 		{
 			SwiftUIView()
 		}
