AVP - visionOS | Improve switching between Fully Immersive and Mixed Mode
-VisionOSRuntimeSettings now control the immersionStyle and upperLimbVisibility properties of the immersive space.
-Establishes a pipeline for pushing data into SwiftMain.swift from Unreal C++, specifically from LaunchIOS.cpp in this case.

diff --git a/Engine/Platforms/VisionOS/Source/Runtime/VisionOSRuntimeSettings/Classes/VisionOSRuntimeSettings.h b/Engine/Platforms/VisionOS/Source/Runtime/VisionOSRuntimeSettings/Classes/VisionOSRuntimeSettings.h
--- a/Engine/Platforms/VisionOS/Source/Runtime/VisionOSRuntimeSettings/Classes/VisionOSRuntimeSettings.h
+++ b/Engine/Platforms/VisionOS/Source/Runtime/VisionOSRuntimeSettings/Classes/VisionOSRuntimeSettings.h
@@ -38,4 +38,12 @@ public:
 	// If checked, the Swift/ObjC bridging headers will be created.
 	UPROPERTY(config, EditAnywhere, Category = Build, Meta = (DisplayName = "Create Swift/ObjC Bridging Headers", ConfigRestartRequired = true))
     bool bCreateBridgingHeader;
+	
+	// Specify the immersiveStyle we will launch into: 0 = .full, 1 = .mixed.
+	UPROPERTY(config, EditAnywhere, Category = ImmersiveSpace, Meta = (DisplayName = "ImmersiveStyle"))
+	int ImmersiveStyle = 0;
+	
+	// Specify the upperLimbVisibility we will start with: 0 = .hidden, 1 = .visible, 2 = .automatic.
+	UPROPERTY(config, EditAnywhere, Category = ImmersiveSpace, Meta = (DisplayName = "UpperLimbVisibility"))
+	int UpperLimbVisibility = 2;
 };
diff --git a/Engine/Source/Runtime/Launch/Private/Apple/SwiftMain.swift b/Engine/Source/Runtime/Launch/Private/Apple/SwiftMain.swift
--- a/Engine/Source/Runtime/Launch/Private/Apple/SwiftMain.swift
+++ b/Engine/Source/Runtime/Launch/Private/Apple/SwiftMain.swift
@@ -10,10 +10,14 @@ struct SwiftUIView: View {
 	@Environment(\.openImmersiveSpace) var openImmersiveSpace
 	@Environment(\.dismissWindow) private var dismissWindow
 
-	@State private var buttonDisabled = false
+	@State private var buttonDisabled = true
 	@State private var timerCounter = 3
 	@State private var timerText = "3"
 	let timer = Timer.publish(every: 1, on: .main, in: .common).autoconnect()
+	@State private var messageText = "Waiting for configuration..."
+	
+	@Binding var style: ImmersionStyle
+	@Binding var limb: Visibility
 	
 	func Open()
 	{
@@ -35,6 +39,7 @@ struct SwiftUIView: View {
 			Text("Unreal SwiftMain Splash Screen!")
 				.font(.title)
 				.bold()
+			Text(messageText)
 			Button(action: {
 				Open()
 			}
@@ -52,6 +57,46 @@ struct SwiftUIView: View {
 				Open()
 			}
 		}
+		.onReceive(NotificationCenter.default.publisher(for: .configureImmersiveSpace)) { notification in
+			guard let userInfo = notification.userInfo else { return }
+			
+			if let styleValue = userInfo["immersiveStyle"] as? Int32 {
+				switch styleValue {
+				case 0:
+					style = .full
+					break
+				case 1:
+					style = .mixed
+					break
+				default:
+					style = .full
+				}
+			}
+
+			if let upperLimbVisibilityValue = userInfo["upperLimbVisibility"] as? Int32 {
+				switch upperLimbVisibilityValue {
+				case 0:
+					limb = .hidden
+					break
+				case 1:
+					limb = .visible
+					break
+#if !UE_SDK_VERSION_1
+// sdk 1 does not support automatic, so fall back to visible
+				case 2:
+					limb = .automatic
+					break
+#endif
+				default:
+					limb = .visible			
+				}
+			}
+			
+			messageText = "ImmersiveStyle: \(String(describing:style))  UpperLimbVisibility: \(String(describing:limb)) "
+			print("messageText: \(messageText)")
+			
+			buttonDisabled = false
+		}
 	}
 }
 
@@ -93,29 +138,37 @@ struct UEContentConfiguration: CompositorLayerConfiguration {
 }
 
 // unused at the moment, but this is how UE can open a SwiftUI view form the Obj-C side
-class HostingViewFactory: NSObject
-{
-	static var onStartImmersive: (LayerRenderer) -> Void = { layer in }
+//class HostingViewFactory: NSObject {
+//	static var onStartImmersive: (LayerRenderer) -> Void = { layer in }
+//
+//	@objc static func MakeSwiftUIView(OnClick: @escaping (() -> Void)) -> UIViewController
+//	{
+//	  return UIHostingController(rootView: SwiftUIView(onClick: OnClick))
+//	}
+//}
 
-	@objc static func MakeSwiftUIView(OnClick: @escaping (() -> Void)) -> UIViewController
-	{
-	  return UIHostingController(rootView: SwiftUIView(onClick: OnClick))
+// Bridge class for calling into this file from c++
+// Data comes in from SwiftMainBridge.cpp and goes out via the swift notification system to SwiftUI elements.
+@objc class SwiftMainBridge: NSObject {
+
+	@objc static func configureImmersiveSpace(inImmersiveStyle: Int32, inUpperLimbVisibility: Int32) {
+		NotificationCenter.default.post(name: .configureImmersiveSpace, 
+										object: nil, 
+										userInfo: ["immersiveStyle": inImmersiveStyle, "upperLimbVisibility": inUpperLimbVisibility])
 	}
 }
 
+extension Notification.Name {
+	static let configureImmersiveSpace = Notification.Name("configureImmersiveSpace")
+}
+
 #if UE_USE_SWIFT_UI_MAIN
 
 @main
 struct UESwiftApp: App {
 	
-#if UE_SDK_VERSION_1
-	// VisionOS 1.x only supports full immersion style in Metal rendering apps.
 	@State private var style: ImmersionStyle = .full
-#else
-	// VisonOS 2+ also supports mixed immersion, which we enable by default.
-	//@State private var style: ImmersionStyle = .full
-	@State private var style: ImmersionStyle = .mixed
-#endif
+	@State private var limb: Visibility = .visible
 	
 	@UIApplicationDelegateAdaptor(IOSAppDelegate.self) var delegate
 	
@@ -123,7 +176,7 @@ struct UESwiftApp: App {
 	{
 		WindowGroup(id: "SwiftUIView")
 		{
-			SwiftUIView()
+			SwiftUIView(style: $style, limb: $limb)
 		}
 
 		ImmersiveSpace(id: "ImmersiveSpace")
@@ -136,16 +189,8 @@ struct UESwiftApp: App {
 					KickoffWithCompositingLayer(layerRenderer)
 			}
 		}
-
-#if UE_SDK_VERSION_1
-		//.upperLimbVisibility(.hidden)
-		.upperLimbVisibility(.visible)
-#else
-		//.upperLimbVisibility(.hidden)
-		//.upperLimbVisibility(.visible)
-		.upperLimbVisibility(.automatic) // Depth occlusion based hand fading, VisionOS 2+ only
-#endif
 		.immersionStyle(selection: $style, in: .mixed, .full)
+		.upperLimbVisibility($limb.wrappedValue)
 	}
 }
 
diff --git a/Engine/Source/Runtime/Launch/Private/Apple/SwiftMainBridge.cpp b/Engine/Source/Runtime/Launch/Private/Apple/SwiftMainBridge.cpp
new file mode 100644
--- /dev/null
+++ b/Engine/Source/Runtime/Launch/Private/Apple/SwiftMainBridge.cpp
@@ -0,0 +1,21 @@
+// Copyright Epic Games, Inc. All Rights Reserved.
+
+#include "SwiftMainBridge.h"
+
+#if PLATFORM_VISIONOS
+
+// This import is the autogenerated "bridging header" created by xcode because SwiftMain.swift contains @objc tagged items.
+// Look in there if you need to understand the objective c interface better.
+#import "SwiftMain-Swift.h"
+
+// All our calls come in from unreal threads, so we dispatch_async to get the into Apple's main thread where SwiftUI updates must happen.
+
+void UE::SwiftMainBridgeNS::ConfigureImmersiveSpace(int32 InImmersiveStyle, int32 InUpperLimbVisibility)
+{
+	dispatch_async(dispatch_get_main_queue(), ^{
+		[SwiftMainBridge configureImmersiveSpaceInImmersiveStyle :InImmersiveStyle inUpperLimbVisibility:InUpperLimbVisibility];
+	});
+
+}
+
+#endif // PLATFORM_VISIONOS
diff --git a/Engine/Source/Runtime/Launch/Private/Apple/SwiftMainBridge.h b/Engine/Source/Runtime/Launch/Private/Apple/SwiftMainBridge.h
new file mode 100644
--- /dev/null
+++ b/Engine/Source/Runtime/Launch/Private/Apple/SwiftMainBridge.h
@@ -0,0 +1,15 @@
+// Copyright Epic Games, Inc. All Rights Reserved.
+
+#pragma once
+
+// Bridge functions for communicating from C++ to SwiftMain.swift.
+// See LaunchCPPToSwift.h for an example of going the other way, from SwiftMain.swift to C++.
+
+#if PLATFORM_VISIONOS
+
+namespace UE::SwiftMainBridgeNS
+{
+	void ConfigureImmersiveSpace(int32 InImmersiveStyle, int32 InUpperLimbVisibility);
+}
+
+#endif // PLATFORM_VISIONOS
diff --git a/Engine/Source/Runtime/Launch/Private/IOS/LaunchIOS.cpp b/Engine/Source/Runtime/Launch/Private/IOS/LaunchIOS.cpp
--- a/Engine/Source/Runtime/Launch/Private/IOS/LaunchIOS.cpp
+++ b/Engine/Source/Runtime/Launch/Private/IOS/LaunchIOS.cpp
@@ -30,6 +30,10 @@
 #include "Misc/EmbeddedCommunication.h"
 #include "Misc/CoreDelegates.h"
 
+#if PLATFORM_VISIONOS && UE_USE_SWIFT_UI_MAIN
+#include "../Apple/SwiftMainBridge.h"
+#endif
+
 FEngineLoop GEngineLoop;
 FGameLaunchDaemonMessageHandler GCommandSystem;
 
@@ -395,8 +399,16 @@ void FAppEntry::PlatformInit()
 
 	// wait until the GLView is fully initialized, so the RHI can be initialized
 	IOSAppDelegate* AppDelegate = [IOSAppDelegate GetDelegate];
-
+	
 #if UE_USE_SWIFT_UI_MAIN
+#if PLATFORM_VISIONOS
+	// Read Immersive space settings and send them to swift where we will open the immersive space.
+	int ImmersiveStyle = 0;
+	GConfig->GetInt(TEXT("/Script/VisionOSRuntimeSettings.VisionOSRuntimeSettings"), TEXT("ImmersiveStyle"), ImmersiveStyle, GEngineIni);
+	int UpperLimbVisibility = 2;
+	GConfig->GetInt(TEXT("/Script/VisionOSRuntimeSettings.VisionOSRuntimeSettings"), TEXT("UpperLimbVisibility"), UpperLimbVisibility, GEngineIni);
+	UE::SwiftMainBridgeNS::ConfigureImmersiveSpace(ImmersiveStyle, UpperLimbVisibility);
+#endif //PLATFORM_VISIONOS
 	GSwiftTriggerEvent->Wait();
 #else
 	
@@ -407,10 +419,10 @@ void FAppEntry::PlatformInit()
 		// processing occur here
 		FTSTicker::GetCoreTicker().Tick(0.005f);
 		FThreadManager::Get().Tick();
-#endif
+#endif //BUILD_EMBEDDED_APP
 		FPlatformProcess::Sleep(0.005f);
 	}
-#endif
+#endif //UE_USE_SWIFT_UI_MAIN
 	// Set GSystemResolution now that we have the size.
 	FDisplayMetrics DisplayMetrics;
 	FDisplayMetrics::RebuildDisplayMetrics(DisplayMetrics);
