[Loader/PreExit]
- The loader will now flush any remaining async loads before PreExit to fix some issues.

diff --git a/Engine/Source/Runtime/Launch/Private/LaunchEngineLoop.cpp b/Engine/Source/Runtime/Launch/Private/LaunchEngineLoop.cpp
--- a/Engine/Source/Runtime/Launch/Private/LaunchEngineLoop.cpp
+++ b/Engine/Source/Runtime/Launch/Private/LaunchEngineLoop.cpp
@@ -4787,6 +4787,11 @@ void FEngineLoop::Exit()
 	}
 #endif // WITH_ENGINE
 
+	// Make sure we're not in the middle of loading something after this point as some
+	// systems are shut down during PreExit and it can cause issues and crashes if loading continues
+	// beyond this point.
+	FlushAsyncLoading();
+
 	if ( GEngine != nullptr )
 	{
 #if WITH_EDITOR
