Drive-by cleanup. Single implementation for NSOperatingSystemVersion comparison.

diff --git a/Engine/Source/Runtime/Core/Private/Mac/MacPlatformMisc.cpp b/Engine/Source/Runtime/Core/Private/Mac/MacPlatformMisc.cpp
--- a/Engine/Source/Runtime/Core/Private/Mac/MacPlatformMisc.cpp
+++ b/Engine/Source/Runtime/Core/Private/Mac/MacPlatformMisc.cpp
@@ -1806,26 +1806,31 @@ bool FMacPlatformMisc::IsRunningOnBattery()
 	return GMacAppInfo.RunningOnBattery;
 }
 
-int32 FMacPlatformMisc::MacOSXVersionCompare(uint8 Major, uint8 Minor, uint8 Revision)
+int32 FMacPlatformMisc::VersionCompare(const NSOperatingSystemVersion& VersionA, const NSOperatingSystemVersion& VersionB)
 {
-	uint8 TargetValues[3] = {Major, Minor, Revision};
-	NSInteger ComponentValues[3] = {GMacAppInfo.OSXVersion.majorVersion, GMacAppInfo.OSXVersion.minorVersion, GMacAppInfo.OSXVersion.patchVersion};
-	
-	for(uint32 i = 0; i < 3; i++)
+	NSInteger ValuesA[3] = { VersionA.majorVersion, VersionA.minorVersion, VersionA.patchVersion };
+	NSInteger ValuesB[3] = { VersionB.majorVersion, VersionB.minorVersion, VersionB.patchVersion };
+
+	for (uint32 i = 0; i < 3; i++)
 	{
-		if(ComponentValues[i] < TargetValues[i])
+		if (ValuesA[i] < ValuesB[i])
 		{
 			return -1;
 		}
-		else if(ComponentValues[i] > TargetValues[i])
+		else if (ValuesA[i] > ValuesB[i])
 		{
 			return 1;
 		}
 	}
-	
+
 	return 0;
 }
 
+int32 FMacPlatformMisc::MacOSXVersionCompare(uint8 Major, uint8 Minor, uint8 Revision)
+{
+	return VersionCompare(GMacAppInfo.OSXVersion, {Major, Minor, Revision});
+}
+
 FString FMacPlatformMisc::GetOperatingSystemId()
 {
 	FString Result;
@@ -1851,22 +1856,7 @@ FString FMacPlatformMisc::GetXcodePath()
 
 int32 FMacPlatformMisc::XCodeVersionCompare(uint8 Major, uint8 Minor, uint8 Revision)
 {
-	uint8 TargetValues[3] = {Major, Minor, Revision};
-	NSInteger ComponentValues[3] = {GMacAppInfo.XcodeVersion.majorVersion, GMacAppInfo.XcodeVersion.minorVersion, GMacAppInfo.XcodeVersion.patchVersion};
-	
-	for(uint32 i = 0; i < 3; i++)
-	{
-		if(ComponentValues[i] < TargetValues[i])
-		{
-			return -1;
-		}
-		else if(ComponentValues[i] > TargetValues[i])
-		{
-			return 1;
-		}
-	}
-	
-	return 0;
+	return VersionCompare(GMacAppInfo.XcodeVersion, {Major, Minor, Revision});
 }
 
 bool FMacPlatformMisc::IsSupportedXcodeVersionInstalled()
diff --git a/Engine/Source/Runtime/Core/Public/Mac/MacPlatformMisc.h b/Engine/Source/Runtime/Core/Public/Mac/MacPlatformMisc.h
--- a/Engine/Source/Runtime/Core/Public/Mac/MacPlatformMisc.h
+++ b/Engine/Source/Runtime/Core/Public/Mac/MacPlatformMisc.h
@@ -102,6 +102,11 @@ struct CORE_API FMacPlatformMisc : public FApplePlatformMisc
 
 	static IPlatformChunkInstall* GetPlatformChunkInstall();
 
+	/**
+	 * Returns if A < B returns -1, if A > B returns 1, else A == B and returns 0.
+	 */
+	static int32 VersionCompare(const NSOperatingSystemVersion& VersionA, const NSOperatingSystemVersion& VersionB);
+
 	/**
 	 * Returns if current < target returns 1, if current > target returns 1, else current == target and returns 0.
 	 */
diff --git a/Engine/Source/Runtime/Launch/Private/Mac/LaunchMac.cpp b/Engine/Source/Runtime/Launch/Private/Mac/LaunchMac.cpp
--- a/Engine/Source/Runtime/Launch/Private/Mac/LaunchMac.cpp
+++ b/Engine/Source/Runtime/Launch/Private/Mac/LaunchMac.cpp
@@ -66,27 +66,6 @@ void EngineFullCrashHandler(const FGenericCrashContext& GenericContext)
     Context.GenerateCrashInfoAndLaunchReporter();
 }
 
-static int32 MacOSVersionCompare(const NSOperatingSystemVersion& VersionA, const NSOperatingSystemVersion& VersionB)
-{
-	NSInteger ValuesA[3] = { VersionA.majorVersion, VersionA.minorVersion, VersionA.patchVersion };
-	NSInteger ValuesB[3] = { VersionB.majorVersion, VersionB.minorVersion, VersionB.patchVersion };
-
-	for (uint32 i = 0; i < 3; i++)
-	{
-		if (ValuesA[i] < ValuesB[i])
-		{
-			return -1;
-		}
-		else if (ValuesA[i] > ValuesB[i])
-		{
-			return 1;
-		}
-	}
-
-	return 0;
-}
-
-
 @implementation UEAppDelegate
 
 -(id)init
@@ -319,13 +298,13 @@ static int32 MacOSVersionCompare(const NSOperatingSystemVersion& VersionA, const
 	MinimumSystemVersion.patchVersion = VersionComponents.count > 2 ? [[VersionComponents objectAtIndex:2] integerValue] : 0;
 
 	// Make sure that the min version in Info.plist is at least 10.14.6, as that's the absolute minimum
-	if (MacOSVersionCompare(MinimumSystemVersion, MinSupportedMacOSVersion) < 0)
+	if (FMacPlatformMisc::VersionCompare(MinimumSystemVersion, MinSupportedMacOSVersion) < 0)
 	{
 		MinimumSystemVersion = MinSupportedMacOSVersion;
 		MinimumSystemVersionString = MinSupportedMacOSVersionString;
 	}
 
-	if (MacOSVersionCompare(CurrentSystemVersion, MinimumSystemVersion) < 0)
+	if (FMacPlatformMisc::VersionCompare(CurrentSystemVersion, MinimumSystemVersion) < 0)
 	{
 		CFDictionaryRef SessionDictionary = CGSessionCopyCurrentDictionary();
 		const bool bIsWindowServerAvailable = SessionDictionary != nullptr;
@@ -380,7 +359,7 @@ static int32 MacOSVersionCompare(const NSOperatingSystemVersion& VersionA, const
 	
 #if !IS_MONOLITHIC
 	// UE-172403: dlopen crash on Ventura [13.0~13.3)
-	if (MacOSVersionCompare(CurrentSystemVersion, {13, 0, 0}) >= 0 && MacOSVersionCompare(CurrentSystemVersion, {13, 3, 0}) < 0)
+	if (FMacPlatformMisc::VersionCompare(CurrentSystemVersion, {13, 0, 0}) >= 0 && FMacPlatformMisc::VersionCompare(CurrentSystemVersion, {13, 3, 0}) < 0)
 	{
 		CFDictionaryRef SessionDictionary = CGSessionCopyCurrentDictionary();
 		const bool bIsWindowServerAvailable = SessionDictionary != nullptr;
