[SCW] Don't attempt to access project binary directory when running as a program (UnrealBuildTool.TargetType.Program).

This fixes a long standing warning "Failed to find game directory" in SCW on Mac and Linux.
This was only showing up on POSIX systems because they forwarded FPlatformMisc::LowLevelOutputDebugStringf() calls to the log.
The directory does not have to be accessed for Program build targets.

#rnx

diff --git a/Engine/Source/Runtime/Launch/Private/LaunchEngineLoop.cpp b/Engine/Source/Runtime/Launch/Private/LaunchEngineLoop.cpp
--- a/Engine/Source/Runtime/Launch/Private/LaunchEngineLoop.cpp
+++ b/Engine/Source/Runtime/Launch/Private/LaunchEngineLoop.cpp
@@ -1836,13 +1836,17 @@ int32 FEngineLoop::PreInitPreStartupScreen(const TCHAR* CmdLine)
 		}
 	}
 
+#if !IS_PROGRAM && !IS_MONOLITHIC
 	FString EngineBinariesRootDirectory = FPlatformMisc::EngineDir();
+#endif
 	FString ProjectBinariesRootDirectory;
 	if (FApp::HasProjectName())
 	{
 		ProjectBinariesRootDirectory = FPlatformMisc::ProjectDir();
 #if !IS_MONOLITHIC
+#if !IS_PROGRAM
 		FPlatformMisc::GetEngineAndProjectAbsoluteDirsFromExecutable(ProjectBinariesRootDirectory, EngineBinariesRootDirectory);
+#endif
 		
 		// Loading preinit/common module if it exists.
 		// PreInit module can contain things like decryption logic for pak files and things that is project specific but needs to initialize very early
