[LaunchEngineLoop] Fix copy paste mistake in error message, it now correctly identifies that it fails to load the Game Engine class instead of the Editor Engine class when instantiating GEngine.

diff --git a/Engine/Source/Runtime/Launch/Private/LaunchEngineLoop.cpp b/Engine/Source/Runtime/Launch/Private/LaunchEngineLoop.cpp
--- a/Engine/Source/Runtime/Launch/Private/LaunchEngineLoop.cpp
+++ b/Engine/Source/Runtime/Launch/Private/LaunchEngineLoop.cpp
@@ -4522,7 +4522,7 @@ int32 FEngineLoop::Init()
 		EngineClass = StaticLoadClass( UGameEngine::StaticClass(), nullptr, *GameEngineClassName);
 		if (EngineClass == nullptr)
 		{
-			UE_LOG(LogInit, Fatal, TEXT("Failed to load UnrealEd Engine class '%s'."), *GameEngineClassName);
+			UE_LOG(LogInit, Fatal, TEXT("Failed to load Game Engine class '%s'."), *GameEngineClassName);
 		}
 		GEngine = NewObject<UEngine>(GetTransientPackage(), EngineClass);
 	}
