Remove every-frame RHI thread flush inadvertently introduced in CL 42586342
- Causing perf regression. The flush is also unnecessary.

diff --git a/Engine/Source/Runtime/Launch/Private/LaunchEngineLoop.cpp b/Engine/Source/Runtime/Launch/Private/LaunchEngineLoop.cpp
--- a/Engine/Source/Runtime/Launch/Private/LaunchEngineLoop.cpp
+++ b/Engine/Source/Runtime/Launch/Private/LaunchEngineLoop.cpp
@@ -5541,7 +5541,6 @@ void FEngineLoop::Tick()
 			[](FRHICommandListImmediate& RHICmdList)
 			{
 				FDeferredUpdateResource::ResetNeedsUpdate();
-				RHICmdList.ImmediateFlush(EImmediateFlushType::FlushRHIThreadFlushResources);
 			});
 
 		// Don't pump messages if we're running embedded as the outer application
