With Pre-Load Movie Player plugin active, the Preload movie would play on startup but then the DefaultGameMoviePlayer would play it again on first map change.
Loading into the Editor with Pre-Load Movie Player plugin active would cause the audio of the configured video to play during editor launch.

diff --git a/Engine/Plugins/Runtime/PreLoadScreenMoviePlayer/Source/PreLoadScreenMoviePlayer/Private/PreLoadMoviePlayerModuleBase.cpp b/Engine/Plugins/Runtime/PreLoadScreenMoviePlayer/Source/PreLoadScreenMoviePlayer/Private/PreLoadMoviePlayerModuleBase.cpp
--- a/Engine/Plugins/Runtime/PreLoadScreenMoviePlayer/Source/PreLoadScreenMoviePlayer/Private/PreLoadMoviePlayerModuleBase.cpp
+++ b/Engine/Plugins/Runtime/PreLoadScreenMoviePlayer/Source/PreLoadScreenMoviePlayer/Private/PreLoadMoviePlayerModuleBase.cpp
@@ -13,7 +13,6 @@ void FPreLoadMoviePlayerScreenModuleBase::RegisterMovieStreamer(TSharedPtr<class
     {
         //Pass the streamer to our screen and then init with this new information
         MoviePreLoadScreen->RegisterMovieStreamer(InMovieStreamer);
-        MoviePreLoadScreen->Init();
     }
 }
 
diff --git a/Engine/Plugins/Runtime/PreLoadScreenMoviePlayer/Source/PreLoadScreenMoviePlayer/Private/PreLoadMoviePlayerScreenBase.cpp b/Engine/Plugins/Runtime/PreLoadScreenMoviePlayer/Source/PreLoadScreenMoviePlayer/Private/PreLoadMoviePlayerScreenBase.cpp
--- a/Engine/Plugins/Runtime/PreLoadScreenMoviePlayer/Source/PreLoadScreenMoviePlayer/Private/PreLoadMoviePlayerScreenBase.cpp
+++ b/Engine/Plugins/Runtime/PreLoadScreenMoviePlayer/Source/PreLoadScreenMoviePlayer/Private/PreLoadMoviePlayerScreenBase.cpp
@@ -170,6 +170,18 @@ void FPreLoadMoviePlayerScreenBase::CleanUp()
 
     MovieStreamer.Reset();
 	MoviePlayerContents.Reset();
+
+	if (MovieViewport.IsValid())
+	{
+		MovieViewport.Reset();
+	}
+
+	MovieStreamingIsDone.Set(1);
+
+	if (MovieAttributes.IsValid())
+	{
+		MovieAttributes = FPreLoadMovieAttributes();
+	}
 }
 
 void FPreLoadMoviePlayerScreenBase::Tick(float DeltaTime)
diff --git a/Engine/Source/Runtime/Launch/Private/LaunchEngineLoop.cpp b/Engine/Source/Runtime/Launch/Private/LaunchEngineLoop.cpp
--- a/Engine/Source/Runtime/Launch/Private/LaunchEngineLoop.cpp
+++ b/Engine/Source/Runtime/Launch/Private/LaunchEngineLoop.cpp
@@ -3313,14 +3313,6 @@ int32 FEngineLoop::PreInitPostStartupScreen(const TCHAR* CmdLine)
 			TSharedPtr<FSlateRenderer> SlateRenderer = PreInitContext.SlateRenderer;
 			TSharedRef<FSlateRenderer> SlateRendererSharedRef = SlateRenderer.ToSharedRef();
 			
-			if (IsMoviePlayerEnabled())
-			{
-				SCOPED_BOOT_TIMING("GetMoviePlayer()->SetupLoadingScreenFromIni");
-				// allow the movie player to load a sequence from the .inis (a PreLoadingScreen module could have already initialized a sequence, in which case
-				// it wouldn't have anything in it's .ini file)
-				GetMoviePlayer()->SetupLoadingScreenFromIni();
-			}
-
 			{
 				SCOPED_BOOT_TIMING("LoadModulesForProject(ELoadingPhase::PreEarlyLoadingScreen)");
 				// Load up all modules that need to hook into the loading screen
@@ -3337,6 +3329,53 @@ int32 FEngineLoop::PreInitPostStartupScreen(const TCHAR* CmdLine)
 			}
 			// If not using the bundle manager, config will be reloaded after ESP, see below
 
+			//If the PreLoadScreen is active, then we don't want to initialise the MoviePlayer from the same data.
+			bool hasPreloadStartupMovie = false;
+			//If true then we want the splash screen to be handled defaultly
+			bool useDefaultSplash = true;
+
+			UE_SCOPED_ENGINE_ACTIVITY(TEXT("Play Preload Screen"));
+			SCOPED_BOOT_TIMING("PlayFirstPreLoadScreen");
+
+			if (FPreLoadScreenManager::Get())
+			{
+				SCOPED_BOOT_TIMING("PlayFirstPreLoadScreen - FPreLoadScreenManager::Get()->Initialize");
+				// initialize and play our first Early PreLoad Screen if one is setup
+				FPreLoadScreenManager::Get()->Initialize(SlateRendererSharedRef.Get());
+
+				//If either early/engine screen is present then the DefaultMoviePlayer should not be configured with the same data.
+				hasPreloadStartupMovie = 
+					FPreLoadScreenManager::Get()->HasRegisteredPreLoadScreenType(EPreLoadScreenTypes::EarlyStartupScreen) ||
+					FPreLoadScreenManager::Get()->HasRegisteredPreLoadScreenType(EPreLoadScreenTypes::EngineLoadingScreen);
+
+				if (FPreLoadScreenManager::Get()->HasRegisteredPreLoadScreenType(EPreLoadScreenTypes::EarlyStartupScreen))
+				{
+					// disable the splash before playing the early startup screen
+					FPreLoadScreenManager::Get()->IsResponsibleForRenderingDelegate.AddLambda(
+						[](bool bIsPreloadScreenManResponsibleForRendering)
+						{
+							FPlatformMisc::PlatformHandleSplashScreen(!bIsPreloadScreenManResponsibleForRendering);
+						}
+					);
+					FPreLoadScreenManager::Get()->PlayFirstPreLoadScreen(EPreLoadScreenTypes::EarlyStartupScreen);
+
+					useDefaultSplash = false;
+				}
+				else
+				{
+					useDefaultSplash = true;
+				}
+			}
+
+			//If a PreLoadScreenManager is present, then MoviePlayer should not read the same configurable data and rely on a different configuration setup.
+			if (IsMoviePlayerEnabled() && !hasPreloadStartupMovie)
+			{
+				SCOPED_BOOT_TIMING("GetMoviePlayer()->SetupLoadingScreenFromIni");
+				// allow the movie player to load a sequence from the .inis (a PreLoadingScreen module could have already initialized a sequence, in which case
+				// it wouldn't have anything in it's .ini file)
+				GetMoviePlayer()->SetupLoadingScreenFromIni();
+			}
+
 			if (IsMoviePlayerEnabled() && GetMoviePlayer()->HasEarlyStartupMovie())
 			{
 				SCOPED_BOOT_TIMING("EarlyStartupMovie");
@@ -3360,40 +3399,14 @@ int32 FEngineLoop::PreInitPostStartupScreen(const TCHAR* CmdLine)
 				}
 				UE_LOG(LogInit, Display, TEXT("\n%s"), *FileList);
 #endif
-			}
-			else
-			{
-				UE_SCOPED_ENGINE_ACTIVITY(TEXT("Play Preload Screen"));
-				SCOPED_BOOT_TIMING("PlayFirstPreLoadScreen");
 
-				if (FPreLoadScreenManager::Get())
-				{
-					SCOPED_BOOT_TIMING("PlayFirstPreLoadScreen - FPreLoadScreenManager::Get()->Initialize");
-					// initialize and play our first Early PreLoad Screen if one is setup
-					FPreLoadScreenManager::Get()->Initialize(SlateRendererSharedRef.Get());
+				useDefaultSplash &= true;
+			}
 
-					if (FPreLoadScreenManager::Get()->HasRegisteredPreLoadScreenType(EPreLoadScreenTypes::EarlyStartupScreen))
-					{
-						// disable the splash before playing the early startup screen
-						FPreLoadScreenManager::Get()->IsResponsibleForRenderingDelegate.AddLambda(
-							[](bool bIsPreloadScreenManResponsibleForRendering)
-							{
-								FPlatformMisc::PlatformHandleSplashScreen(!bIsPreloadScreenManResponsibleForRendering);
-							}
-						);
-						FPreLoadScreenManager::Get()->PlayFirstPreLoadScreen(EPreLoadScreenTypes::EarlyStartupScreen);
-					}
-					else
-					{
-						// no early startup screen, show the splash screen
-						FPlatformMisc::PlatformHandleSplashScreen(true);
-					}
-				}
-				else
-				{
-					// no preload manager, show the splash screen
-					FPlatformMisc::PlatformHandleSplashScreen(true);
-				}
+			if (useDefaultSplash)
+			{
+				// no early startup screen, show the splash screen
+				FPlatformMisc::PlatformHandleSplashScreen(true);
 			}
 		}
 		else if (IsRunningCommandlet())
