- No longer have StorageServerPlatformFile set itself as the Physical PlatfomrFile
- Changed the  AssetDataGatherer to ask the Physical PlatformFile instead of the topmost PlatformFile for Journaling support (which is a physical driver check)

diff --git a/Engine/Source/Runtime/AssetRegistry/Private/AssetDataGatherer.cpp b/Engine/Source/Runtime/AssetRegistry/Private/AssetDataGatherer.cpp
--- a/Engine/Source/Runtime/AssetRegistry/Private/AssetDataGatherer.cpp
+++ b/Engine/Source/Runtime/AssetRegistry/Private/AssetDataGatherer.cpp
@@ -209,7 +209,7 @@ void FGatherSettings::Initialize()
 		static_cast<void>(FTelemetryRouter::Get());
 	}
 
-	bool bPlatformSupportsDiscoveryCacheInvalidation = FPlatformFileManager::Get().GetPlatformFile().FileJournalIsAvailable();
+	bool bPlatformSupportsDiscoveryCacheInvalidation = FPlatformFileManager::Get().GetPlatformPhysical().FileJournalIsAvailable();
 
 	bool bSkipInvalidate = FParse::Param(FCommandLine::Get(), TEXT("AssetRegistryCacheSkipInvalidate"));
 	EFeatureEnabledReadWrite DiscoverySetting = EFeatureEnabledReadWrite::Invalid;
diff --git a/Engine/Source/Runtime/Launch/Private/LaunchEngineLoop.cpp b/Engine/Source/Runtime/Launch/Private/LaunchEngineLoop.cpp
--- a/Engine/Source/Runtime/Launch/Private/LaunchEngineLoop.cpp
+++ b/Engine/Source/Runtime/Launch/Private/LaunchEngineLoop.cpp
@@ -778,7 +778,13 @@ bool LaunchCheckForFileOverride(const TCHAR* CmdLine, bool& OutFileOverrideFound
 		{
 			CurrentPlatformFile = NetworkPlatformFile;
 			FPlatformFileManager::Get().SetPlatformFile(*CurrentPlatformFile);
-			FPlatformFileManager::Get().SetPlatformPhysical(*CurrentPlatformFile);
+			// pak streaming currently needs this to be the PlatformPhysical, although it can potentially
+			// cause issues for other code that assumes PlatformPhysical only refers to the absolute lowest
+			// level for the platform
+            if (UE::IsUsingZenPakFileStreaming())
+            {
+                FPlatformFileManager::Get().SetPlatformPhysical(*CurrentPlatformFile);
+            }
 		}
 	}
 
