[CoreDelegates]
- Cleanup deprecation

diff --git a/Engine/Source/Runtime/Launch/Private/Android/LaunchAndroid.cpp b/Engine/Source/Runtime/Launch/Private/Android/LaunchAndroid.cpp
--- a/Engine/Source/Runtime/Launch/Private/Android/LaunchAndroid.cpp
+++ b/Engine/Source/Runtime/Launch/Private/Android/LaunchAndroid.cpp
@@ -1928,10 +1928,6 @@ static void OnAppCommandCB(struct android_app* app, int32_t cmd)
 		{
 			FGraphEventRef WillTerminateTask = FFunctionGraphTask::CreateAndDispatchWhenReady([]()
 			{
-				PRAGMA_DISABLE_DEPRECATION_WARNINGS
-				FCoreDelegates::ApplicationWillTerminateDelegate.Broadcast();
-				PRAGMA_ENABLE_DEPRECATION_WARNINGS
-
 				FCoreDelegates::GetApplicationWillTerminateDelegate().Broadcast();
 			}, TStatId(), NULL, ENamedThreads::GameThread);
 			FTaskGraphInterface::Get().WaitUntilTaskCompletes(WillTerminateTask);
