Added CPU profiling scopes for more FCoreDelegate broadcasts (OnBeginFrame, OnBeginFrameRT, OnEndFrameRT; OnEndFrame already had the CPU scope added).

diff --git a/Engine/Source/Runtime/Launch/Private/LaunchEngineLoop.cpp b/Engine/Source/Runtime/Launch/Private/LaunchEngineLoop.cpp
--- a/Engine/Source/Runtime/Launch/Private/LaunchEngineLoop.cpp
+++ b/Engine/Source/Runtime/Launch/Private/LaunchEngineLoop.cpp
@@ -5309,7 +5309,11 @@ static inline void BeginFrameRenderThread(FRHICommandListImmediate& RHICmdList,
 #endif
 
 	GPU_STATS_BEGINFRAME(RHICmdList);
-	FCoreDelegates::OnBeginFrameRT.Broadcast();
+
+	{
+		QUICK_SCOPE_CYCLE_COUNTER(STAT_FEngineLoop_CoreDelegates_BeginFrameRT);
+		FCoreDelegates::OnBeginFrameRT.Broadcast();
+	}
 
 	RHICmdList.EnqueueLambda([CurrentFrameCounter](FRHICommandListImmediate& InRHICmdList)
 	{
@@ -5344,7 +5348,11 @@ static inline void EndFrameRenderThread(FRHICommandListImmediate& RHICmdList, ui
 		UEngine::SetRenderSubmitLatencyMarkerEnd(CurrentFrameCounter);
 	});
 
-	FCoreDelegates::OnEndFrameRT.Broadcast();
+	{
+		QUICK_SCOPE_CYCLE_COUNTER(STAT_FEngineLoop_CoreDelegates_EndFrameRT);
+		FCoreDelegates::OnEndFrameRT.Broadcast();
+	}
+
 	GVisualizeTexture.EndFrameRenderThread();
 
 	GPU_STATS_ENDFRAME(RHICmdList);
@@ -5488,7 +5496,10 @@ void FEngineLoop::Tick()
 		UpdateCoreCsvStats_BeginFrame();
 #endif
 
-		FCoreDelegates::OnBeginFrame.Broadcast();
+		{
+			QUICK_SCOPE_CYCLE_COUNTER(STAT_FEngineLoop_CoreDelegates_BeginFrame);
+			FCoreDelegates::OnBeginFrame.Broadcast();
+		}
 
 		// flush debug output which has been buffered by other threads
 		{
