Add profiler scopes around FEngineLoop code that was not being properly tracked.
Convert QUICK_SCOPE_CYCLE_COUNTER to SCOPED_NAMED_EVENT so they show up in in platform profilers and insights even without -statnamedevents.
This means they will always be visible in Test configuration profiles.

diff --git a/Engine/Source/Runtime/Launch/Private/LaunchEngineLoop.cpp b/Engine/Source/Runtime/Launch/Private/LaunchEngineLoop.cpp
--- a/Engine/Source/Runtime/Launch/Private/LaunchEngineLoop.cpp
+++ b/Engine/Source/Runtime/Launch/Private/LaunchEngineLoop.cpp
@@ -318,6 +318,12 @@ UE_TRACE_EVENT_BEGIN(Cpu, Frame, NoSync)
 UE_TRACE_EVENT_END()
 #endif // CPUPROFILERTRACE_ENABLED
 
+#if ENABLE_NAMED_EVENTS
+// Different levels of color for scoped named events 
+static const FColor FrameColor = FColor::Red;
+static const FColor SystemColor = FColor::Orange;
+#endif // ENABLE_NAMED_EVENTS
+
 bool GIsConsoleExecutable = false;
 
 static TAutoConsoleVariable<int32> CVarDoAsyncEndOfFrameTasksRandomize(
@@ -1343,6 +1349,7 @@ static void UpdateCoreCsvStats_BeginFrame()
 {
 	if (FCsvProfiler::Get()->IsCapturing())
 	{
+		TRACE_CPUPROFILER_EVENT_SCOPE(UpdateCoreCsvStats_BeginFrame);
 		if (!IsRunningDedicatedServer())
 		{
 #if !UE_BUILD_SHIPPING
@@ -1397,6 +1404,7 @@ CSV_DEFINE_CATEGORY(GPUUsage, true);
 
 static void UpdateCoreCsvStats_EndFrame()
 {
+	TRACE_CPUPROFILER_EVENT_SCOPE(UpdateCoreCsvStats_EndFrame);
 	if (!IsRunningDedicatedServer())
 	{
 		CSV_CUSTOM_STAT_MINIMAL_GLOBAL(RenderThreadTime, FPlatformTime::ToMilliseconds(GRenderThreadTime), ECsvCustomStatOp::Set);
@@ -4758,7 +4766,7 @@ namespace PipelineStateCache
 void FEngineLoop::Exit()
 {
 	STAT_ADD_CUSTOMMESSAGE_NAME( STAT_NamedMarker, TEXT( "EngineLoop.Exit" ) );
-	TRACE_CPUPROFILER_EVENT_SCOPE(FEngineLoop::Exit);
+	SCOPED_NAMED_EVENT_TEXT("FEngineLoop::Exit", FrameColor);
 	TRACE_BOOKMARK(TEXT("EngineLoop.Exit"));
 
 	ClearPendingCleanupObjects();
@@ -5311,7 +5319,7 @@ static inline void BeginFrameRenderThread(FRHICommandListImmediate& RHICmdList,
 	GPU_STATS_BEGINFRAME(RHICmdList);
 
 	{
-		QUICK_SCOPE_CYCLE_COUNTER(STAT_FEngineLoop_CoreDelegates_BeginFrameRT);
+		SCOPED_NAMED_EVENT(FCoreDelegates_BeginFrameRT, SystemColor);
 		FCoreDelegates::OnBeginFrameRT.Broadcast();
 	}
 
@@ -5349,7 +5357,7 @@ static inline void EndFrameRenderThread(FRHICommandListImmediate& RHICmdList, ui
 	});
 
 	{
-		QUICK_SCOPE_CYCLE_COUNTER(STAT_FEngineLoop_CoreDelegates_EndFrameRT);
+		SCOPED_NAMED_EVENT(FCoreDelegates_EndFrameRT, SystemColor);
 		FCoreDelegates::OnEndFrameRT.Broadcast();
 	}
 
@@ -5373,7 +5381,7 @@ static inline void EndFrameRenderThread(FRHICommandListImmediate& RHICmdList, ui
 
 void FEngineLoop::Tick()
 {
-	TRACE_CPUPROFILER_EVENT_SCOPE(FEngineLoop::Tick);
+	SCOPED_NAMED_EVENT_TEXT("FEngineLoop::Tick", FrameColor);
 	SCOPE_STALL_COUNTER(FEngineLoop::Tick, 2.0);
 
 #if !UE_BUILD_SHIPPING && !UE_BUILD_TEST && MALLOC_GT_HOOKS
@@ -5394,15 +5402,15 @@ void FEngineLoop::Tick()
 #endif
 
 	{
-		TRACE_CPUPROFILER_EVENT_SCOPE(HeartBeat);
+		SCOPED_NAMED_EVENT(FEngineLoop_HeartBeat, SystemColor);
 		// Send a heartbeat for the diagnostics thread
 		FThreadHeartBeat::Get().HeartBeat(true);
-	}
 
-	FGameThreadHitchHeartBeat::Get().FrameStart();
+		FGameThreadHitchHeartBeat::Get().FrameStart();
+	}
 
 	{
-		TRACE_CPUPROFILER_EVENT_SCOPE(TickHotfixables);
+		SCOPED_NAMED_EVENT(Tick_Hotfixables, SystemColor);
 		FPlatformMisc::TickHotfixables();
 	}
 
@@ -5412,7 +5420,7 @@ void FEngineLoop::Tick()
 	// Make sure something is ticking the rendering tickables in -onethread mode to avoid leaks/bugs.
 	if (!GUseThreadedRendering)
 	{
-		TRACE_CPUPROFILER_EVENT_SCOPE(TickRenderingTickables);
+		SCOPED_NAMED_EVENT(Tick_RenderingTickables, SystemColor);
 		TickRenderingTickables();
 	}
 
@@ -5463,12 +5471,12 @@ void FEngineLoop::Tick()
 #if PLATFORM_LIMIT_PROFILER_UNIQUE_NAMED_EVENTS
 	FrameString = TEXT("FEngineLoop");
 #endif
-	SCOPED_NAMED_EVENT_TCHAR_CONDITIONAL(FrameString, FColor::Red, !bTraceCpuChannelEnabled); 
+	SCOPED_NAMED_EVENT_TCHAR_CONDITIONAL(FrameString, FrameColor, !bTraceCpuChannelEnabled); 
 #endif //ENABLE_NAMED_EVENTS
 
 	// execute callbacks for cvar changes
 	{
-		QUICK_SCOPE_CYCLE_COUNTER(STAT_FEngineLoop_Tick_CallAllConsoleVariableSinks);
+		SCOPED_NAMED_EVENT(FEngineLoop_CallAllConsoleVariableSinks, SystemColor);
 		IConsoleManager::Get().CallAllConsoleVariableSinks();
 	}
 
@@ -5497,13 +5505,13 @@ void FEngineLoop::Tick()
 #endif
 
 		{
-			QUICK_SCOPE_CYCLE_COUNTER(STAT_FEngineLoop_CoreDelegates_BeginFrame);
+			SCOPED_NAMED_EVENT(FCoreDelegates_BeginFrame, SystemColor);
 			FCoreDelegates::OnBeginFrame.Broadcast();
 		}
 
 		// flush debug output which has been buffered by other threads
 		{
-			QUICK_SCOPE_CYCLE_COUNTER(STAT_FEngineLoop_FlushThreadedLogs); 
+			SCOPED_NAMED_EVENT(FEngineLoop_FlushThreadedLogs, SystemColor);
 			GLog->FlushThreadedLogs(EOutputDeviceRedirectorFlushOptions::Async);
 		}
 
@@ -5516,38 +5524,41 @@ void FEngineLoop::Tick()
 
 		// set FApp::CurrentTime, FApp::DeltaTime and potentially wait to enforce max tick rate
 		{
-			QUICK_SCOPE_CYCLE_COUNTER(STAT_FEngineLoop_UpdateTimeAndHandleMaxTickRate);
+			SCOPED_NAMED_EVENT(FEngineLoop_UpdateTimeAndHandleMaxTickRate, SystemColor);
 			GEngine->UpdateTimeAndHandleMaxTickRate();
 			GEngine->SetSimulationLatencyMarkerStart(CurrentFrameCounter);
 		}
 
 		// beginning of RHI frame
-		ENQUEUE_RENDER_COMMAND(BeginFrame)([CurrentFrameCounter](FRHICommandListImmediate& RHICmdList)
-		{
-			BeginFrameRenderThread(RHICmdList, CurrentFrameCounter);
-		});
-
-		for (FSceneInterface* Scene : GetRendererModule().GetAllocatedScenes())
 		{
-			ENQUEUE_RENDER_COMMAND(FScene_StartFrame)([Scene](FRHICommandListImmediate& RHICmdList)
+			SCOPED_NAMED_EVENT(Render_BeginFrame, SystemColor);
+			ENQUEUE_RENDER_COMMAND(BeginFrame)([CurrentFrameCounter](FRHICommandListImmediate& RHICmdList)
 			{
-				Scene->StartFrame();
+				BeginFrameRenderThread(RHICmdList, CurrentFrameCounter);
 			});
-		}
 
-		UE::RenderCommandPipe::StartRecording();
+			for (FSceneInterface* Scene : GetRendererModule().GetAllocatedScenes())
+			{
+				ENQUEUE_RENDER_COMMAND(FScene_StartFrame)([Scene](FRHICommandListImmediate& RHICmdList)
+				{
+					Scene->StartFrame();
+				});
+			}
+
+			UE::RenderCommandPipe::StartRecording();
 
 #if !UE_SERVER && WITH_ENGINE
-		if (!GIsEditor && GEngine->GameViewport && GEngine->GameViewport->GetWorld() && GEngine->GameViewport->GetWorld()->IsCameraMoveable())
-		{
-			// When not in editor, we emit dynamic resolution's begin frame right after RHI's.
-			GEngine->EmitDynamicResolutionEvent(EDynamicResolutionStateEvent::BeginFrame);
-		}
+			if (!GIsEditor && GEngine->GameViewport && GEngine->GameViewport->GetWorld() && GEngine->GameViewport->GetWorld()->IsCameraMoveable())
+			{
+				// When not in editor, we emit dynamic resolution's begin frame right after RHI's.
+				GEngine->EmitDynamicResolutionEvent(EDynamicResolutionStateEvent::BeginFrame);
+			}
 #endif
+		}
 
 		// tick performance monitoring
 		{
-			QUICK_SCOPE_CYCLE_COUNTER(STAT_FEngineLoop_TickFPSChart);
+			SCOPED_NAMED_EVENT(Tick_PerformanceMonitoring, SystemColor);
 			GEngine->TickPerformanceMonitoring( FApp::GetDeltaTime() );
 
 			extern COREUOBJECT_API void ResetAsyncLoadingStats();
@@ -5557,7 +5568,7 @@ void FEngineLoop::Tick()
 #if UPDATE_MALLOC_STATS
 		// update memory allocator stats
 		{
-			QUICK_SCOPE_CYCLE_COUNTER(STAT_FEngineLoop_Malloc_UpdateStats);
+			SCOPED_NAMED_EVENT(Malloc_UpdateStats, SystemColor);
 			UE::Private::GMalloc->UpdateStats();
 		}
 #endif
@@ -5584,14 +5595,13 @@ void FEngineLoop::Tick()
 		{
 			GEngine->SetInputSampleLatencyMarker(CurrentFrameCounter);
 
-			//QUICK_SCOPE_CYCLE_COUNTER(STAT_PumpMessages);
+			SCOPED_NAMED_EVENT(FEngineLoop_PumpMessages, SystemColor);
 			FPlatformApplicationMisc::PumpMessages(true);
 		}
 
 		bool bIdleMode;
 		{
-
-			QUICK_SCOPE_CYCLE_COUNTER(STAT_FEngineLoop_Idle);
+			SCOPED_NAMED_EVENT(FEngineLoop_Idle, SystemColor);
 
 			// Idle mode prevents ticking and rendering completely
 			bIdleMode = ShouldUseIdleMode();
@@ -5640,23 +5650,26 @@ void FEngineLoop::Tick()
 		{
 			CSV_SCOPED_TIMING_STAT_EXCLUSIVE(Input);
 			SCOPE_TIME_GUARD(TEXT("SlateInput"));
-			QUICK_SCOPE_CYCLE_COUNTER(STAT_FEngineLoop_Tick_SlateInput);
+			SCOPED_NAMED_EVENT(Tick_SlateInput, SystemColor);
 			LLM_SCOPE(ELLMTag::UI);
 
 			FSlateApplication& SlateApp = FSlateApplication::Get();
             {
-                QUICK_SCOPE_CYCLE_COUNTER(STAT_FEngineLoop_Tick_PollGameDeviceState);
+				TRACE_CPUPROFILER_EVENT_SCOPE(PollGameDeviceState);
                 SlateApp.PollGameDeviceState();
             }
 			// Gives widgets a chance to process any accumulated input
             {
-                QUICK_SCOPE_CYCLE_COUNTER(STAT_FEngineLoop_Tick_FinishedInputThisFrame);
+				TRACE_CPUPROFILER_EVENT_SCOPE(FinishedInputThisFrame);
                 SlateApp.FinishedInputThisFrame();
             }
 		}
 		
-		// main game engine tick (world, game objects, etc.)
-		GEngine->Tick(FApp::GetDeltaTime(), bIdleMode);
+		{
+			SCOPED_NAMED_EVENT(Tick_Engine, SystemColor);
+			// main game engine tick (world, game objects, etc.)
+			GEngine->Tick(FApp::GetDeltaTime(), bIdleMode);
+		}
 
 		// If a movie that is blocking the game thread has been playing,
 		// wait for it to finish before we continue to tick or tick again
@@ -5697,7 +5710,7 @@ void FEngineLoop::Tick()
             }
 			else if (IsMoviePlayerEnabled())
 			{
-				QUICK_SCOPE_CYCLE_COUNTER(STAT_FEngineLoop_WaitForMovieToFinish);
+				TRACE_CPUPROFILER_EVENT_SCOPE(WaitForMovieToFinish);
 				GetMoviePlayer()->WaitForMovieToFinish(true);
 			}
 		}
@@ -5710,7 +5723,7 @@ void FEngineLoop::Tick()
 		if (FSlateApplication::IsInitialized() && !bIdleMode)
 		{
 			{
-				QUICK_SCOPE_CYCLE_COUNTER(STAT_FEngineLoop_ProcessPlayerControllersSlateOperations);
+				SCOPED_NAMED_EVENT(ProcessLocalPlayerSlateOperations, SystemColor);
 				check(!IsRunningDedicatedServer());
 
 				// Process slate operations accumulated in the world ticks.
@@ -5762,7 +5775,7 @@ void FEngineLoop::Tick()
 				{
 					CSV_SCOPED_SET_WAIT_STAT(ConcurrentWithSlate);
 
-					QUICK_SCOPE_CYCLE_COUNTER(STAT_ConcurrentWithSlateTickTasks_Wait);
+					TRACE_CPUPROFILER_EVENT_SCOPE(ConcurrentWithSlateTickTasks_Wait);
 					check(ConcurrentTaskCompleteEvent);
 					ConcurrentTaskCompleteEvent->Wait();
 					FPlatformProcess::ReturnSynchEventToPool(ConcurrentTaskCompleteEvent);
@@ -5805,7 +5818,7 @@ void FEngineLoop::Tick()
 		{
 			CSV_SCOPED_SET_WAIT_STAT(ConcurrentWithSlate);
 
-			QUICK_SCOPE_CYCLE_COUNTER(STAT_ConcurrentWithSlateTickTasks_Wait);
+			TRACE_CPUPROFILER_EVENT_SCOPE(ConcurrentWithSlateTickTasks_Wait);
 			check(ConcurrentTaskCompleteEvent);
 			ConcurrentTaskCompleteEvent->Wait();
 			FPlatformProcess::ReturnSynchEventToPool(ConcurrentTaskCompleteEvent);
@@ -5823,7 +5836,7 @@ void FEngineLoop::Tick()
 #if WITH_EDITOR && !UE_BUILD_SHIPPING
 		// tick automation controller (Editor only)
 		{
-			QUICK_SCOPE_CYCLE_COUNTER(STAT_FEngineLoop_Tick_AutomationController);
+			SCOPED_NAMED_EVENT(Tick_AutomationController, SystemColor);
 			static FName AutomationController("AutomationController");
 			if (FModuleManager::Get().IsModuleLoaded(AutomationController))
 			{
@@ -5835,7 +5848,7 @@ void FEngineLoop::Tick()
 #if WITH_ENGINE && WITH_AUTOMATION_WORKER
 		// tick automation worker
 		{
-			QUICK_SCOPE_CYCLE_COUNTER(STAT_FEngineLoop_Tick_AutomationWorker);
+			SCOPED_NAMED_EVENT(Tick_AutomationWorker, SystemColor);
 			static const FName AutomationWorkerModuleName = TEXT("AutomationWorker");
 			if (FModuleManager::Get().IsModuleLoaded(AutomationWorkerModuleName))
 			{
@@ -5844,20 +5857,24 @@ void FEngineLoop::Tick()
 		}
 #endif
 
-		UE::RenderCommandPipe::StopRecording();
-
-		for (FSceneInterface* Scene : GetRendererModule().GetAllocatedScenes())
 		{
-			ENQUEUE_RENDER_COMMAND(FScene_EndFrame)([Scene](FRHICommandListImmediate& RHICmdList)
+			SCOPED_NAMED_EVENT(Render_UpdateFrame, SystemColor);
+
+			UE::RenderCommandPipe::StopRecording();
+
+			for (FSceneInterface* Scene : GetRendererModule().GetAllocatedScenes())
 			{
-				Scene->EndFrame(RHICmdList);
-			});
-		}
+				ENQUEUE_RENDER_COMMAND(FScene_EndFrame)([Scene](FRHICommandListImmediate& RHICmdList)
+				{
+					Scene->EndFrame(RHICmdList);
+				});
+			}
 
-		// tick render hardware interface
-		{			
-			SCOPE_CYCLE_COUNTER(STAT_RHITickTime);
-			RHITick( FApp::GetDeltaTime() ); // Update RHI.
+			// tick render hardware interface
+			{			
+				SCOPE_CYCLE_COUNTER(STAT_RHITickTime);
+				RHITick( FApp::GetDeltaTime() ); // Update RHI.
+			}
 		}
 
 		// We need to set this marker before EndFrameRenderThread is enqueued. 
@@ -5883,15 +5900,26 @@ void FEngineLoop::Tick()
 			SCOPE_CYCLE_COUNTER(STAT_DeferredTickTime);
 			CSV_SCOPED_TIMING_STAT_EXCLUSIVE(DeferredTickTime);
 			// Delete the objects which were enqueued for deferred cleanup before the previous frame.
-			delete PreviousPendingCleanupObjects;
 
+			{
+				SCOPED_NAMED_EVENT(DeletePendingCleanupObjects, SystemColor);
+
+				delete PreviousPendingCleanupObjects;
 #if WITH_COREUOBJECT
-			DeleteLoaders(); // destroy all linkers pending delete
+				DeleteLoaders(); // destroy all linkers pending delete
 #endif
+			}
 
-			FTSTicker::GetCoreTicker().Tick(FApp::GetDeltaTime());
-			FThreadManager::Get().Tick();
-			GEngine->TickDeferredCommands();		
+			{
+				SCOPED_NAMED_EVENT(Tick_Core, SystemColor);
+				FTSTicker::GetCoreTicker().Tick(FApp::GetDeltaTime());
+				FThreadManager::Get().Tick();
+			}
+
+			{
+				SCOPED_NAMED_EVENT(Tick_DeferredCommands, SystemColor);
+				GEngine->TickDeferredCommands();
+			}
 		}
 		FMoviePlayerProxy::BlockingForceFinished();
 
@@ -5901,13 +5929,13 @@ void FEngineLoop::Tick()
 		IMediaModule* MediaModule = FModuleManager::LoadModulePtr<IMediaModule>(MediaModuleName);
 		if (MediaModule != nullptr)
 		{
-			QUICK_SCOPE_CYCLE_COUNTER(STAT_FEngineLoop_MediaTickPostRender);
+			SCOPED_NAMED_EVENT(Tick_PostRender, SystemColor);
 			MediaModule->TickPostRender();
 		}
 #endif
 
 		{
-			QUICK_SCOPE_CYCLE_COUNTER(STAT_FEngineLoop_CoreDelegates_EndFrame);
+			SCOPED_NAMED_EVENT(FCoreDelegates_EndFrame, SystemColor);
 			FCoreDelegates::OnEndFrame.Broadcast();
 		}
 
@@ -5928,24 +5956,28 @@ void FEngineLoop::Tick()
 		UpdateCoreCsvStats_EndFrame();
 #endif
 
-		// Tick DumpGPU to start/stop frame dumps
-		#if WITH_ENGINE && WITH_DUMPGPU
-			UE::RenderCore::DumpGPU::TickEndFrame();
-		#endif
-
-		#if !UE_SERVER && WITH_ENGINE
 		{
-			// We emit dynamic resolution's end frame right before RHI's. GEngine is going to ignore it if no BeginFrame was done.
-			GEngine->EmitDynamicResolutionEvent(EDynamicResolutionStateEvent::EndFrame);
-		}
-		#endif
+			SCOPED_NAMED_EVENT(Render_EndFrame, SystemColor);
 
-		// end of RHI frame
-		ENQUEUE_RENDER_COMMAND(EndFrame)(
-			[CurrentFrameCounter](FRHICommandListImmediate& RHICmdList)
+			// Tick DumpGPU to start/stop frame dumps
+#if WITH_ENGINE && WITH_DUMPGPU
+				UE::RenderCore::DumpGPU::TickEndFrame();
+#endif
+
+#if !UE_SERVER && WITH_ENGINE
 			{
-				EndFrameRenderThread(RHICmdList, CurrentFrameCounter);
-			});
+				// We emit dynamic resolution's end frame right before RHI's. GEngine is going to ignore it if no BeginFrame was done.
+				GEngine->EmitDynamicResolutionEvent(EDynamicResolutionStateEvent::EndFrame);
+			}
+#endif
+
+			// end of RHI frame
+			ENQUEUE_RENDER_COMMAND(EndFrame)(
+				[CurrentFrameCounter](FRHICommandListImmediate& RHICmdList)
+				{
+					EndFrameRenderThread(RHICmdList, CurrentFrameCounter);
+				});
+		}
 
 		// Set CPU utilization stats.
 		const FCPUTime CPUTime = FPlatformTime::GetCPUTime();
