Fix PVS warnings

#rnx

diff --git a/Engine/Plugins/Messaging/UdpMessaging/Source/UdpMessaging/Private/Transport/UdpMessageTransport.cpp b/Engine/Plugins/Messaging/UdpMessaging/Source/UdpMessaging/Private/Transport/UdpMessageTransport.cpp
--- a/Engine/Plugins/Messaging/UdpMessaging/Source/UdpMessaging/Private/Transport/UdpMessageTransport.cpp
+++ b/Engine/Plugins/Messaging/UdpMessaging/Source/UdpMessaging/Private/Transport/UdpMessageTransport.cpp
@@ -321,6 +321,7 @@ bool FUdpMessageTransport::StartTransport(IMessageTransportHandler& Handler)
 #if PLATFORM_DESKTOP
 	MessageProcessor = new FUdpMessageProcessor(*UnicastSocket, FGuid::NewGuid(), MulticastEndpoint);
 #else
+	check(MulticastSocket);
 	MessageProcessor = new FUdpMessageProcessor(*MulticastSocket, FGuid::NewGuid(), MulticastEndpoint);
 #endif
 	{
diff --git a/Engine/Source/Runtime/D3D12RHI/Private/D3D12StateCachePrivate.h b/Engine/Source/Runtime/D3D12RHI/Private/D3D12StateCachePrivate.h
--- a/Engine/Source/Runtime/D3D12RHI/Private/D3D12StateCachePrivate.h
+++ b/Engine/Source/Runtime/D3D12RHI/Private/D3D12StateCachePrivate.h
@@ -692,6 +692,7 @@ public:
 	{
 		if (RTArray) //NULL is legal
 		{
+			check(NumSimultaneousRTs);
 			FMemory::Memcpy(RTArray, PipelineState.Graphics.RenderTargetArray, sizeof(FD3D12RenderTargetView*)* D3D12_SIMULTANEOUS_RENDER_TARGET_COUNT);
 			*NumSimultaneousRTs = PipelineState.Graphics.CurrentNumberOfRenderTargets;
 		}
diff --git a/Engine/Source/Runtime/Engine/Private/InstancedStaticMesh.cpp b/Engine/Source/Runtime/Engine/Private/InstancedStaticMesh.cpp
--- a/Engine/Source/Runtime/Engine/Private/InstancedStaticMesh.cpp
+++ b/Engine/Source/Runtime/Engine/Private/InstancedStaticMesh.cpp
@@ -1232,6 +1232,7 @@ FInstancedStaticMeshVFLooseUniformShaderParametersRef FInstancedStaticMeshSceneP
 		// InstancedLODRange is only set for HierarchicalInstancedStaticMeshes
 		if (InstancingUserData && InstancedLODRange)
 		{
+			check(View);
 			int32 FirstLOD = InstancingUserData->MinLOD;
 
 			int32 DebugMin = FMath::Min(CVarMinLOD.GetValueOnRenderThread(), InstancingUserData->MeshRenderData->LODResources.Num() - 1);
diff --git a/Engine/Source/Runtime/Engine/Private/RepLayout.cpp b/Engine/Source/Runtime/Engine/Private/RepLayout.cpp
--- a/Engine/Source/Runtime/Engine/Private/RepLayout.cpp
+++ b/Engine/Source/Runtime/Engine/Private/RepLayout.cpp
@@ -4073,7 +4073,6 @@ bool FRepLayout::ReceiveProperties_BackwardsCompatible_r(
 			bool bShadowDataCopied = false;
 			TArray<FProperty*>* RepNotifies = (ShadowData && RepState) ? &RepState->RepNotifies : nullptr;
 
-
 			FGuidReferencesMap* NewGuidReferencesArray = PrepReceivedArray(
 				ArrayNum,
 				ShadowArray,
@@ -4162,13 +4161,14 @@ bool FRepLayout::ReceiveProperties_BackwardsCompatible_r(
 		{
 			const int32 ElementOffset = (Data.Data - OldData.Data);
 
+			TArray<FProperty*>* RepNotifies = (ShadowData && RepState) ? &RepState->RepNotifies : nullptr;
 			if (ReceivePropertyHelper(
 				TempReader,
 				GuidReferencesMap,
 				ElementOffset,
 				ShadowData,
 				Data,
-				ShadowData ? &RepState->RepNotifies : nullptr,
+				RepNotifies,
 				false /*bShadowDataCopied*/,
 				Parents,
 				Cmds,
diff --git a/Engine/Source/Runtime/Launch/Private/LaunchEngineLoop.cpp b/Engine/Source/Runtime/Launch/Private/LaunchEngineLoop.cpp
--- a/Engine/Source/Runtime/Launch/Private/LaunchEngineLoop.cpp
+++ b/Engine/Source/Runtime/Launch/Private/LaunchEngineLoop.cpp
@@ -825,7 +825,7 @@ bool LaunchCheckForFileOverride(const TCHAR* CmdLine, bool& OutFileOverrideFound
 #if PLATFORM_REQUIRES_FILESERVER
 			FPlatformMisc::LowLevelOutputDebugStringf(TEXT("Failed to connect to file server at %s. RETRYING in 5s.\n"), *HostIpString);
 			FPlatformProcess::Sleep(5.0f);
-			uint32 Result = 2;
+			constexpr uint32 Result = 2;
 #else	//PLATFORM_REQUIRES_FILESERVER
 			// note that this can't be localized because it happens before we connect to a filserver - localizing would cause ICU to try to load.... from over the file server connection!
 			FString Error = FString::Printf(TEXT("Failed to connect to any of the following file servers:\n\n    %s\n\nWould you like to try again? No will fallback to local disk files, Cancel will quit."), *HostIpString.Replace( TEXT("+"), TEXT("\n    ")));
