Added CPU profiler scope events (SCOPED_NAMED_EVENT) for FCoreDelegates::OnXXX.Broadcast(); calls from LaunchEngineLoop.

diff --git a/Engine/Source/Runtime/Launch/Private/LaunchEngineLoop.cpp b/Engine/Source/Runtime/Launch/Private/LaunchEngineLoop.cpp
--- a/Engine/Source/Runtime/Launch/Private/LaunchEngineLoop.cpp
+++ b/Engine/Source/Runtime/Launch/Private/LaunchEngineLoop.cpp
@@ -322,6 +322,7 @@ UE_TRACE_EVENT_END()
 // Different levels of color for scoped named events 
 static const FColor FrameColor = FColor::Red;
 static const FColor SystemColor = FColor::Orange;
+static const FColor CoreDelegatesColor = FColor::Cyan;
 #endif // ENABLE_NAMED_EVENTS
 
 bool GIsConsoleExecutable = false;
@@ -464,14 +465,22 @@ static FAutoConsoleCommand CmdRunCommandlet(
 
 			PRIVATE_GRunningCommandletClass = CommandletClass;
 
-			FCoreDelegates::OnCommandletPreMain.Broadcast();
+			{
+				SCOPED_NAMED_EVENT_TEXT("FCoreDelegates::OnCommandletPreMain", CoreDelegatesColor);
+				FCoreDelegates::OnCommandletPreMain.Broadcast();
+			}
+
 			int32 ErrorLevel;
 			{
 				TRACE_CPUPROFILER_EVENT_SCOPE_TEXT(*WriteToString<512>(TEXT("Commandlet Main "), Commandlet->GetFName()));
 				FTrackedActivityScope CommandletActivity(FTrackedActivity::GetEngineActivity(), *FString::Printf(TEXT("Running %s"), *Commandlet->GetName()), false, FTrackedActivity::ELight::Green);
 				ErrorLevel = Commandlet->Main(CommandletCommandLine);
 			}
-			FCoreDelegates::OnCommandletPostMain.Broadcast();
+
+			{
+				SCOPED_NAMED_EVENT_TEXT("FCoreDelegates::OnCommandletPostMain", CoreDelegatesColor);
+				FCoreDelegates::OnCommandletPostMain.Broadcast();
+			}
 
 			PRIVATE_GRunningCommandletClass = nullptr;
 
@@ -1905,7 +1914,10 @@ int32 FEngineLoop::PreInitPreStartupScreen(const TCHAR* CmdLine)
 		GWarn = FPlatformOutputDevices::GetFeedbackContext();
 #endif
 
-		FCoreDelegates::OnOutputDevicesInit.Broadcast();
+		{
+			SCOPED_NAMED_EVENT_TEXT("FCoreDelegates::OnOutputDevicesInit", CoreDelegatesColor);
+			FCoreDelegates::OnOutputDevicesInit.Broadcast();
+		}
 	}
 
 	// Avoiding potential exploits by not exposing command line overrides in the shipping games.
@@ -3854,7 +3866,10 @@ int32 FEngineLoop::PreInitPostStartupScreen(const TCHAR* CmdLine)
 			}
 
 			// Call init callbacks
-			FCoreDelegates::OnPostEngineInit.Broadcast();
+			{
+				SCOPED_NAMED_EVENT_TEXT("FCoreDelegates::OnPostEngineInit", CoreDelegatesColor);
+				FCoreDelegates::OnPostEngineInit.Broadcast();
+			}
 
 			// Load all the post-engine init modules
 			ensure(IProjectManager::Get().LoadModulesForProject(ELoadingPhase::PostEngineInit));
@@ -3889,14 +3904,24 @@ int32 FEngineLoop::PreInitPostStartupScreen(const TCHAR* CmdLine)
 #if WITH_ENGINE
 			PRIVATE_GRunningCommandletClass = CommandletClass;
 #endif
-			FCoreDelegates::OnCommandletPreMain.Broadcast();
+
+			{
+				SCOPED_NAMED_EVENT_TEXT("FCoreDelegates::OnCommandletPreMain", CoreDelegatesColor);
+				FCoreDelegates::OnCommandletPreMain.Broadcast();
+			}
+
 			int32 ErrorLevel;
 			{
 				TRACE_CPUPROFILER_EVENT_SCOPE_TEXT(*WriteToString<512>(TEXT("Commandlet Main "), Commandlet->GetFName()));
 				FTrackedActivityScope CommandletActivity(FTrackedActivity::GetEngineActivity(), *FString::Printf(TEXT("Running %s"), *Commandlet->GetName()), false, FTrackedActivity::ELight::Green);
 				ErrorLevel = Commandlet->Main(CommandletCommandLine);
 			}
-			FCoreDelegates::OnCommandletPostMain.Broadcast();
+
+			{
+				SCOPED_NAMED_EVENT_TEXT("FCoreDelegates::OnCommandletPostMain", CoreDelegatesColor);
+				FCoreDelegates::OnCommandletPostMain.Broadcast();
+			}
+
 #if WITH_ENGINE
 			PRIVATE_GRunningCommandletClass = nullptr;
 #endif
@@ -4062,6 +4087,7 @@ int32 FEngineLoop::PreInitPostStartupScreen(const TCHAR* CmdLine)
 #if WITH_EDITOR
 			if (GIsEditor)
 			{
+				SCOPED_NAMED_EVENT_TEXT("FCoreDelegates::OnEditorPreExit", CoreDelegatesColor);
 				FEditorDelegates::OnEditorPreExit.Broadcast();
 			}
 #endif // WITH_EDITOR
@@ -4610,7 +4636,7 @@ int32 FEngineLoop::Init()
 
 	// Call init callbacks
 	{
-		SCOPED_BOOT_TIMING("OnPostEngineInit.Broadcast");
+		SCOPED_BOOT_TIMING("FCoreDelegates::OnPostEngineInit");
 		FCoreDelegates::OnPostEngineInit.Broadcast();
 	}
 
@@ -4729,7 +4755,7 @@ int32 FEngineLoop::Init()
     }
     
     {
-		SCOPED_BOOT_TIMING("FCoreDelegates::OnFEngineLoopInitComplete.Broadcast()");
+		SCOPED_BOOT_TIMING("FCoreDelegates::OnFEngineLoopInitComplete");
 		FCoreDelegates::OnFEngineLoopInitComplete.Broadcast();
 	}
 	FShaderPipelineCache::ResumeBatching();
@@ -4832,6 +4858,7 @@ void FEngineLoop::Exit()
 #if WITH_EDITOR
 		if (GIsEditor)
 		{
+			SCOPED_NAMED_EVENT_TEXT("FCoreDelegates::OnEditorPreExit", CoreDelegatesColor);
 			FEditorDelegates::OnEditorPreExit.Broadcast();
 		}
 #endif // WITH_EDITOR
@@ -5319,7 +5346,7 @@ static inline void BeginFrameRenderThread(FRHICommandListImmediate& RHICmdList,
 	GPU_STATS_BEGINFRAME(RHICmdList);
 
 	{
-		SCOPED_NAMED_EVENT(FCoreDelegates_BeginFrameRT, SystemColor);
+		SCOPED_NAMED_EVENT_TEXT("FCoreDelegates::OnBeginFrameRT", CoreDelegatesColor);
 		FCoreDelegates::OnBeginFrameRT.Broadcast();
 	}
 
@@ -5357,7 +5384,7 @@ static inline void EndFrameRenderThread(FRHICommandListImmediate& RHICmdList, ui
 	});
 
 	{
-		SCOPED_NAMED_EVENT(FCoreDelegates_EndFrameRT, SystemColor);
+		SCOPED_NAMED_EVENT_TEXT("FCoreDelegates::OnEndFrameRT", CoreDelegatesColor);
 		FCoreDelegates::OnEndFrameRT.Broadcast();
 	}
 
@@ -5505,7 +5532,7 @@ void FEngineLoop::Tick()
 #endif
 
 		{
-			SCOPED_NAMED_EVENT(FCoreDelegates_BeginFrame, SystemColor);
+			SCOPED_NAMED_EVENT_TEXT("FCoreDelegates::OnBeginFrame", CoreDelegatesColor);
 			FCoreDelegates::OnBeginFrame.Broadcast();
 		}
 
@@ -5643,7 +5670,10 @@ void FEngineLoop::Tick()
 		FPlatformFileManager::Get().TickActivePlatformFile();
 
 		// Roughly track the time when the input was sampled
-		FCoreDelegates::OnSamplingInput.Broadcast();
+		{
+			SCOPED_NAMED_EVENT_TEXT("FCoreDelegates::OnSamplingInput", CoreDelegatesColor);
+			FCoreDelegates::OnSamplingInput.Broadcast();
+		}
 
 		// process accumulated Slate input
 		if (FSlateApplication::IsInitialized() && !bIdleMode)
@@ -5935,7 +5965,7 @@ void FEngineLoop::Tick()
 #endif
 
 		{
-			SCOPED_NAMED_EVENT(FCoreDelegates_EndFrame, SystemColor);
+			SCOPED_NAMED_EVENT_TEXT("FCoreDelegates::OnEndFrame", CoreDelegatesColor);
 			FCoreDelegates::OnEndFrame.Broadcast();
 		}
 
@@ -6636,7 +6666,7 @@ bool FEngineLoop::AppInit( )
 
 	// Init other systems.
 	{
-		SCOPED_BOOT_TIMING("FCoreDelegates::OnInit.Broadcast");
+		SCOPED_BOOT_TIMING("FCoreDelegates::OnInit");
 		FCoreDelegates::OnInit.Broadcast();
 	}
 
@@ -6662,7 +6692,10 @@ void FEngineLoop::AppPreExit( )
 
 	UE_LOG(LogExit, Log, TEXT("Preparing to exit.") );
 
-	FCoreDelegates::OnPreExit.Broadcast();
+	{
+		SCOPED_NAMED_EVENT_TEXT("FCoreDelegates::OnPreExit", CoreDelegatesColor);
+		FCoreDelegates::OnPreExit.Broadcast();
+	}
 
 #if WITH_ENGINE
 	if (FString(FCommandLine::Get()).Contains(TEXT("CreatePak")) && GetDerivedDataCache())
@@ -6677,7 +6710,10 @@ void FEngineLoop::AppPreExit( )
 	FRemoteConfig::Flush();
 #endif
 
-	FCoreDelegates::OnExit.Broadcast();
+	{
+		SCOPED_NAMED_EVENT_TEXT("FCoreDelegates::OnExit", CoreDelegatesColor);
+		FCoreDelegates::OnExit.Broadcast();
+	}
 
 	// FGenericRHIGPUFence uses GFrameNumberRenderThread to tell when a frame is finished. If such an object is added in the last frame before we
 	// exit, it will never be "signaled", since nothing ever increments GFrameNumberRenderThread again. This can lead to deadlocks in code
