Generate shaders headers in PreInitPreStartupScreen.
At start up, some modules are loading packages and those packages trigger the calculation of shader dependencies.But the shader headers are not generated yet. Shaders heders are generated in the cook commandlet but it's too late.

#rnx

diff --git a/Engine/Source/Runtime/Launch/Private/LaunchEngineLoop.cpp b/Engine/Source/Runtime/Launch/Private/LaunchEngineLoop.cpp
--- a/Engine/Source/Runtime/Launch/Private/LaunchEngineLoop.cpp
+++ b/Engine/Source/Runtime/Launch/Private/LaunchEngineLoop.cpp
@@ -1600,6 +1600,32 @@ static FString ExtractGameStringArgument(TArray<FString>& TokenArray, TArray<FSt
 	return FString();
 }
 
+#if WITH_EDITOR
+static void GenerateShaderHeaders()
+{
+	FShaderCompileUtilities::GenerateBrdfHeaders(GMaxRHIShaderPlatform);
+
+	//Generate shader headers for all the shader platforms. It's necessary for the cook. I don't do it in the cook commandlet
+	//because it's not early enough.
+	{
+		ITargetPlatformManagerModule* PlatformModule = GetTargetPlatformManager();
+		if (PlatformModule)
+		{
+			const TArray<ITargetPlatform*>& ActivePlatforms = PlatformModule->GetActiveTargetPlatforms();
+			for (const ITargetPlatform* Platform : ActivePlatforms)
+			{
+				TArray<FName> ShaderPlatformNameArray;
+				Platform->GetAllTargetedShaderFormats(ShaderPlatformNameArray);
+
+				for (const FName& ShaderPlatformName : ShaderPlatformNameArray)
+				{
+					FShaderCompileUtilities::GenerateBrdfHeaders(ShaderPlatformName);
+				}
+			}
+		}
+	}
+}
+#endif // #if WITH_EDITOR
 
 DECLARE_CYCLE_STAT(TEXT("FEngineLoop::PreInitPreStartupScreen.AfterStats"), STAT_FEngineLoop_PreInitPreStartupScreen_AfterStats, STATGROUP_LoadTime);
 DECLARE_CYCLE_STAT(TEXT("FEngineLoop::PreInitPostStartupScreen.AfterStats"), STAT_FEngineLoop_PreInitPostStartupScreen_AfterStats, STATGROUP_LoadTime);
@@ -3117,7 +3143,7 @@ int32 FEngineLoop::PreInitPreStartupScreen(const TCHAR* CmdLine)
 #if WITH_EDITOR
 			// Explicitly generate AutogenShaderHeaders.ush prior to shader type initialization
 			// (since that process will load and cache this header as a sideeffect)
-			FShaderCompileUtilities::GenerateBrdfHeaders(GMaxRHIShaderPlatform);
+			GenerateShaderHeaders();
 #endif
 
 			// Initialize shader types before loading any shaders
