Generate shader headers after ProcessNewlyLoadedUObjects otherwise some console variables are not available and a warning is printed.

#rnx

diff --git a/Engine/Source/Runtime/Launch/Private/LaunchEngineLoop.cpp b/Engine/Source/Runtime/Launch/Private/LaunchEngineLoop.cpp
--- a/Engine/Source/Runtime/Launch/Private/LaunchEngineLoop.cpp
+++ b/Engine/Source/Runtime/Launch/Private/LaunchEngineLoop.cpp
@@ -1604,6 +1604,28 @@ static FString ExtractGameStringArgument(TArray<FString>& TokenArray, TArray<FSt
 	return FString();
 }
 
+#if WITH_EDITOR
+static void GenerateShaderHeaders()
+{
+	//Generate shader headers for all the shader platforms. It's necessary for the cook. I don't do it in the cook commandlet
+	//because it's not early enough.
+	ITargetPlatformManagerModule* PlatformModule = GetTargetPlatformManager();
+	if (PlatformModule)
+	{
+		const TArray<ITargetPlatform*>& ActivePlatforms = PlatformModule->GetActiveTargetPlatforms();
+		for (const ITargetPlatform* Platform : ActivePlatforms)
+		{
+			TArray<FName> ShaderPlatformNameArray;
+			Platform->GetAllTargetedShaderFormats(ShaderPlatformNameArray);
+
+			for (const FName& ShaderPlatformName : ShaderPlatformNameArray)
+			{
+				FShaderCompileUtilities::GenerateBrdfHeaders(ShaderPlatformName);
+			}
+		}
+	}
+}
+#endif // #if WITH_EDITOR
 
 DECLARE_CYCLE_STAT(TEXT("FEngineLoop::PreInitPreStartupScreen.AfterStats"), STAT_FEngineLoop_PreInitPreStartupScreen_AfterStats, STATGROUP_LoadTime);
 DECLARE_CYCLE_STAT(TEXT("FEngineLoop::PreInitPostStartupScreen.AfterStats"), STAT_FEngineLoop_PreInitPostStartupScreen_AfterStats, STATGROUP_LoadTime);
@@ -3540,6 +3562,16 @@ int32 FEngineLoop::PreInitPostStartupScreen(const TCHAR* CmdLine)
 		}
 
 #if WITH_EDITOR
+
+		// To make sure the shader headers are generated correctly, it needs to run after ProcessNewlyLoadedUObjects. ProcessNewlyLoadedUObjects sets up
+		// the DeviceProfileManager, which is necessary to access certain console variables.
+		// Also, the shader header needs to be generated before the call to GEditor->InitEditor because some modules (at least Niagara) will generate
+		// and cache shaders dependencies during plugin loading so I need the shader headers at that moment.
+		if (AllowShaderCompiling())
+		{
+			GenerateShaderHeaders();
+		}
+
 		if (!UE::Virtualization::ShouldInitializePreSlate())
 		{
 			// Explicit initialization of the virtualization system, after slate has initialized and we can show error dialogs.
