Ensure keyboard OnTextCommitted messages are dispatched when using Android improved virtual keyboard. Ensure nativeKeyboardInputResult is invoked when keyboard is dismissed and keep reference to VirtualKeyboardWidget valid until the result is handled.

diff --git a/Engine/Build/Android/Java/src/com/epicgames/unreal/GameActivity.java.template b/Engine/Build/Android/Java/src/com/epicgames/unreal/GameActivity.java.template
--- a/Engine/Build/Android/Java/src/com/epicgames/unreal/GameActivity.java.template
+++ b/Engine/Build/Android/Java/src/com/epicgames/unreal/GameActivity.java.template
@@ -5954,6 +5954,9 @@ Trace.endSection();
 					InputMethodManager imm =(InputMethodManager)getSystemService(Context.INPUT_METHOD_SERVICE);
 					imm.hideSoftInputFromWindow(newVirtualKeyboardInput.getWindowToken(), 0);
 
+					String finalMessage = newVirtualKeyboardInput.getText().toString();
+					nativeVirtualKeyboardResult(true, finalMessage);
+					
 					nativeVirtualKeyboardVisible(false);
 					bKeyboardShowing = false;
 				}
diff --git a/Engine/Source/Runtime/Launch/Private/Android/AndroidJNI.cpp b/Engine/Source/Runtime/Launch/Private/Android/AndroidJNI.cpp
--- a/Engine/Source/Runtime/Launch/Private/Android/AndroidJNI.cpp
+++ b/Engine/Source/Runtime/Launch/Private/Android/AndroidJNI.cpp
@@ -814,14 +814,12 @@ JNI_METHOD void Java_com_epicgames_unreal_GameActivity_nativeVirtualKeyboardVisi
 	//remove reference so the object can be clicked again to show the virtual keyboard
 	if (!bShown)
 	{
-		VirtualKeyboardWidget.Reset();
-
 		if (FTaskGraphInterface::IsRunning())
 		{
 			FGraphEventRef VirtualKeyboardShown = FFunctionGraphTask::CreateAndDispatchWhenReady([&]()
 			{
 				FAndroidApplication::Get()->OnVirtualKeyboardHidden().Broadcast();
-			}, TStatId(), NULL, ENamedThreads::GameThread);
+        	}, TStatId(), NULL, ENamedThreads::GameThread);
 		}
 	}
 }
@@ -1148,9 +1146,6 @@ void AndroidThunkCpp_HideVirtualKeyboardInput()
 
 	if (JNIEnv* Env = FAndroidApplication::GetJavaEnv())
 	{
-		// ignore anything it might return
-		VirtualKeyboardWidget.Reset();
-
 		// call the java side
 		FJavaWrapper::CallVoidMethod(Env, FJavaWrapper::GameActivityThis, FJavaWrapper::AndroidThunkJava_HideVirtualKeyboardInput);
 	}
