Fix shipping config not loading common module when using UE_MERGED_MODULES
#rnx

diff --git a/Engine/Source/Runtime/Launch/Private/LaunchEngineLoop.cpp b/Engine/Source/Runtime/Launch/Private/LaunchEngineLoop.cpp
--- a/Engine/Source/Runtime/Launch/Private/LaunchEngineLoop.cpp
+++ b/Engine/Source/Runtime/Launch/Private/LaunchEngineLoop.cpp
@@ -1896,10 +1896,37 @@ int32 FEngineLoop::PreInitPreStartupScreen(const TCHAR* CmdLine)
 		// Loading preinit/common module if it exists.
 		// PreInit module can contain things like decryption logic for pak files and things that is project specific but needs to initialize very early
 		// Common module is a merged binary containing all the important modules for the project
+
+#if UE_MERGED_MODULES
+		// Strip suffix "-Platform-Config" to allow building module name to load: "App" + "-ModuleName" + "-Platform-Config"
+		const FString ExecutableNameWithoutExtension = FPlatformProcess::ExecutableName();
+		const FStringView BaseExecutableName = FPathViews::GetBaseFilename(ExecutableNameWithoutExtension);
+		FStringView BasePrefix = BaseExecutableName;
+		FString BaseSuffix;
+		if (FApp::GetBuildConfiguration() != EBuildConfiguration::Development)
+		{
+			int32 LastDash = INDEX_NONE;
+			if (BasePrefix.FindLastChar(TEXT('-'), LastDash))
+			{
+				BasePrefix.LeftInline(LastDash);
+
+				LastDash = INDEX_NONE;
+				if (BasePrefix.FindLastChar(TEXT('-'), LastDash))
+				{
+					BasePrefix.LeftInline(LastDash);
+				}
+			}
+
+			BaseSuffix = BaseExecutableName.RightChop(BasePrefix.Len());
+		}
+
+		const FString PathPrefix = FPaths::Combine(FPathViews::GetPath(ExecutableNameWithoutExtension), BasePrefix);
+#endif
+
 		TArray<FString> FileNames =
 		{
 #if UE_MERGED_MODULES
-			FString::Printf(TEXT("%s-Common.%s"), FPlatformProcess::ExecutableName(), FPlatformProcess::GetModuleExtension()),
+			FString::Printf(TEXT("%s-Common%s.%s"), *PathPrefix, *BaseSuffix, FPlatformProcess::GetModuleExtension()),
 			FString::Printf(TEXT("%s-Common-%s-%s.%s"), FPlatformProcess::ExecutableName(), FPlatformProcess::GetBinariesSubdirectory(), LexToString(FApp::GetBuildConfiguration()), FPlatformProcess::GetModuleExtension()),
 #else
 			FString::Printf(TEXT("%s-%sPreInit.%s"), FPlatformProcess::ExecutableName(), FApp::GetProjectName(), FPlatformProcess::GetModuleExtension()),
