ExeName-(ProjectName)PreInit-Win64-Shipping.dll can now be found and loaded when not running Development configuration

Was looking previously looking for a file that did not exist:
ExeName-Win64-Shipping-(ProjectName)PreInit.dll

#rnx

diff --git a/Engine/Source/Runtime/Launch/Private/LaunchEngineLoop.cpp b/Engine/Source/Runtime/Launch/Private/LaunchEngineLoop.cpp
--- a/Engine/Source/Runtime/Launch/Private/LaunchEngineLoop.cpp
+++ b/Engine/Source/Runtime/Launch/Private/LaunchEngineLoop.cpp
@@ -1893,11 +1893,6 @@ int32 FEngineLoop::PreInitPreStartupScreen(const TCHAR* CmdLine)
 		FPlatformMisc::GetEngineAndProjectAbsoluteDirsFromExecutable(ProjectBinariesRootDirectory, EngineBinariesRootDirectory);
 #endif
 		
-		// Loading preinit/common module if it exists.
-		// PreInit module can contain things like decryption logic for pak files and things that is project specific but needs to initialize very early
-		// Common module is a merged binary containing all the important modules for the project
-
-#if UE_MERGED_MODULES
 		// Strip suffix "-Platform-Config" to allow building module name to load: "App" + "-ModuleName" + "-Platform-Config"
 		const FString ExecutableNameWithoutExtension = FPlatformProcess::ExecutableName();
 		const FStringView BaseExecutableName = FPathViews::GetBaseFilename(ExecutableNameWithoutExtension);
@@ -1921,15 +1916,17 @@ int32 FEngineLoop::PreInitPreStartupScreen(const TCHAR* CmdLine)
 		}
 
 		const FString PathPrefix = FPaths::Combine(FPathViews::GetPath(ExecutableNameWithoutExtension), BasePrefix);
-#endif
 
+		// Loading preinit/common module if it exists.
+		// PreInit module can contain things like decryption logic for pak files and things that is project specific but needs to initialize very early
+		// Common module is a merged binary containing all the important modules for the project
 		TArray<FString> FileNames =
 		{
 #if UE_MERGED_MODULES
 			FString::Printf(TEXT("%s-Common%s.%s"), *PathPrefix, *BaseSuffix, FPlatformProcess::GetModuleExtension()),
 			FString::Printf(TEXT("%s-Common-%s-%s.%s"), FPlatformProcess::ExecutableName(), FPlatformProcess::GetBinariesSubdirectory(), LexToString(FApp::GetBuildConfiguration()), FPlatformProcess::GetModuleExtension()),
 #else
-			FString::Printf(TEXT("%s-%sPreInit.%s"), FPlatformProcess::ExecutableName(), FApp::GetProjectName(), FPlatformProcess::GetModuleExtension()),
+			FString::Printf(TEXT("%s-%sPreInit%s.%s"), *PathPrefix, FApp::GetProjectName(), *BaseSuffix, FPlatformProcess::GetModuleExtension()),
 			FString::Printf(TEXT("%s-%sPreInit-%s-%s.%s"), *FApp::GetName(), FApp::GetProjectName(), FPlatformProcess::GetBinariesSubdirectory(), LexToString(FApp::GetBuildConfiguration()), FPlatformProcess::GetModuleExtension())
 #endif
 		};
