Corrupt pak files validation

When detecting corrupt installation files on startup, show a dialog to let the user validate files automatically.

diff --git a/Engine/Source/Runtime/Core/Private/Misc/CoreDelegates.cpp b/Engine/Source/Runtime/Core/Private/Misc/CoreDelegates.cpp
--- a/Engine/Source/Runtime/Core/Private/Misc/CoreDelegates.cpp
+++ b/Engine/Source/Runtime/Core/Private/Misc/CoreDelegates.cpp
@@ -56,6 +56,7 @@ TTSMulticastDelegate<void(const IPakFile&)>& FCoreDelegates::GetOnPakFileUnmount
 
 TMulticastDelegate<void(const FString&)> FCoreDelegates::NewFileAddedDelegate;
 TMulticastDelegate<void()> FCoreDelegates::NoPakFilesMountedDelegate;
+TMulticastDelegate<void(const FString&)> FCoreDelegates::PakFileMountFailedDelegate;
 
 TTSMulticastDelegate<void(const TCHAR*, const TCHAR*)>& FCoreDelegates::GetOnFileOpenedForReadFromPakFile()
 {
diff --git a/Engine/Source/Runtime/Core/Public/Misc/CoreDelegates.h b/Engine/Source/Runtime/Core/Public/Misc/CoreDelegates.h
--- a/Engine/Source/Runtime/Core/Public/Misc/CoreDelegates.h
+++ b/Engine/Source/Runtime/Core/Public/Misc/CoreDelegates.h
@@ -154,6 +154,9 @@ public:
 	// After an attempt to mount all pak files, but none wre found, this is called
 	static CORE_API TMulticastDelegate<void()> NoPakFilesMountedDelegate;
 
+	// Called when a pak file cannot be mounted
+	static CORE_API TMulticastDelegate<void(const FString&)> PakFileMountFailedDelegate;
+
 	// When a file is opened for read from a pak file.
 	static CORE_API TTSMulticastDelegate<void(const TCHAR* PakFile, const TCHAR* FileName)>& GetOnFileOpenedForReadFromPakFile();
 
diff --git a/Engine/Source/Runtime/Launch/Private/LaunchEngineLoop.cpp b/Engine/Source/Runtime/Launch/Private/LaunchEngineLoop.cpp
--- a/Engine/Source/Runtime/Launch/Private/LaunchEngineLoop.cpp
+++ b/Engine/Source/Runtime/Launch/Private/LaunchEngineLoop.cpp
@@ -494,6 +494,70 @@ static FAutoConsoleCommand CmdRunCommandlet(
 	);
 #endif
 
+#if defined(WITH_LAUNCHERCHECK) && defined(WITH_LAUNCHERCHECK_PAKVERIFY) && WITH_LAUNCHERCHECK && WITH_LAUNCHERCHECK_PAKVERIFY
+struct FPakValidator
+{
+	explicit FPakValidator() :
+		bPakMountFailed(false)
+	{
+		Handle = FCoreDelegates::PakFileMountFailedDelegate.AddLambda([this](const FString&) { bPakMountFailed = true; });
+	}
+
+	~FPakValidator()
+	{
+		FCoreDelegates::PakFileMountFailedDelegate.Remove(Handle);
+	}
+
+	bool HasPakMountFailed() const { return bPakMountFailed; }
+
+	static void HandlePakMountError()
+	{
+		bool bRunVerifyInstallation = false;
+		if (ILauncherCheckModule::Get().WasRanFromLauncher())
+		{
+			bRunVerifyInstallation = FMessageDialog::Open(
+				EAppMsgCategory::Error,
+				EAppMsgType::YesNo,
+				EAppReturnType::No,
+				LOCTEXT("Error.CorruptInstall.VerifyLauncher", "Corrupt files detected. \n\nDo you want to verify your installation?"),
+				LOCTEXT("Error.CorruptInstall.Title", "Corrupt files detected")
+			) == EAppReturnType::Yes;
+		}
+		else
+		{
+			FMessageDialog::Open(
+				EAppMsgCategory::Error,
+				EAppMsgType::Ok,
+				LOCTEXT("Error.CorruptInstall", "Corrupt files detected. Please verify your installation."),
+				LOCTEXT("Error.CorruptInstall.Title", "Corrupt files detected")
+			);
+		}
+
+		if (bRunVerifyInstallation)
+		{
+			ILauncherCheckModule::Get().RunLauncher(ELauncherAction::AppVerifyInstallation);
+			RequestEngineExit(TEXT("Corrupt files detected; running verify in launcher"));
+		}
+		else
+		{
+			RequestEngineExit(TEXT("Corrupt files detected; exiting"));
+		}
+	}
+
+private:
+	FDelegateHandle Handle;
+	bool bPakMountFailed;
+};
+
+#else
+struct FPakValidator 
+{
+	static void HandlePakMountError() {}
+	bool HasPakMountFailed() const { return false; }
+};
+
+#endif
+
 #if WITH_APPLICATION_CORE
 static TUniquePtr<FOutputDeviceConsole>	GScopedLogConsole;
 #endif
@@ -859,8 +923,14 @@ bool LaunchCheckForFileOverride(const TCHAR* CmdLine, bool& OutFileOverrideFound
 #endif
 	if (UE::IsUsingZenPakFileStreaming() || !NetworkPlatformFile)
 	{
+		FPakValidator Validator;
 		IPlatformFile* PlatformFile = ConditionallyCreateFileWrapper(TEXT("PakFile"), CurrentPlatformFile, CmdLine);
-		if (PlatformFile)
+		if (Validator.HasPakMountFailed())
+		{
+			FPakValidator::HandlePakMountError();
+			return false;
+		}
+		else if (PlatformFile)
 		{
 			CurrentPlatformFile = PlatformFile;
 			FPlatformFileManager::Get().SetPlatformFile(*CurrentPlatformFile);
@@ -1757,7 +1827,7 @@ int32 FEngineLoop::PreInitPreStartupScreen(const TCHAR* CmdLine)
 #endif
 	FCoreDelegates::OnSamplingInput.AddStatic(UpdateGInputTime);
 
-#if defined(WITH_LAUNCHERCHECK) && WITH_LAUNCHERCHECK
+#if defined(WITH_LAUNCHERCHECK) && WITH_LAUNCHERCHECK && !WITH_EDITOR
 	if (ILauncherCheckModule::Get().WasRanFromLauncher() == false)
 	{
 		// Tell Launcher to run us instead
diff --git a/Engine/Source/Runtime/PakFile/Private/IPlatformFilePak.cpp b/Engine/Source/Runtime/PakFile/Private/IPlatformFilePak.cpp
--- a/Engine/Source/Runtime/PakFile/Private/IPlatformFilePak.cpp
+++ b/Engine/Source/Runtime/PakFile/Private/IPlatformFilePak.cpp
@@ -6069,13 +6069,15 @@ bool FPakPlatformFile::Mount(const FPakMountArgs& MountArgs, FIoStatus* OutIoMou
 							else
 							{
 								UE_LOG(LogPakFile, Warning, TEXT("Failed to mount optional segment extension IoStore container \"%s\" [%s]"), *OptionalSegmentUtocPath, *MountResult.Status().ToString());
+								FCoreDelegates::PakFileMountFailedDelegate.Broadcast(OptionalSegmentUtocPath);
 							}
 						}
 #endif
 					}
 					else
 					{
-						UE_LOG(LogPakFile, Warning, TEXT("Failed to mount IoStore container \"%s\" [%s]"), *UtocPath, *MountResult.Status().ToString());
+						UE_LOG(LogPakFile, Warning, TEXT("Failed to mount IoStore container \"%s\" [%s]"), *UtocPath, *MountResult.Status().ToString()); 
+						FCoreDelegates::PakFileMountFailedDelegate.Broadcast(InPakFilename);
 					}
 				}
 				else
@@ -6084,6 +6086,7 @@ bool FPakPlatformFile::Mount(const FPakMountArgs& MountArgs, FIoStatus* OutIoMou
 					{
 						IoStoreSuccess = FIoStatus(EIoErrorCode::NotFound, TEXT("utoc not found"));
 						UE_LOG(LogPakFile, Warning, TEXT("IoStore container \"%s\" not found"), *UtocPath);
+						FCoreDelegates::PakFileMountFailedDelegate.Broadcast(UtocPath);
 					}
 					else
 					{
