Destroy Preload system post-init if nothing loaded for it, stops conflicts with default movie player.

diff --git a/Engine/Source/Runtime/Launch/Private/LaunchEngineLoop.cpp b/Engine/Source/Runtime/Launch/Private/LaunchEngineLoop.cpp
--- a/Engine/Source/Runtime/Launch/Private/LaunchEngineLoop.cpp
+++ b/Engine/Source/Runtime/Launch/Private/LaunchEngineLoop.cpp
@@ -3322,16 +3322,15 @@ int32 FEngineLoop::PreInitPreStartupScreen(const TCHAR* CmdLine)
 
 					if (FPreLoadScreenManager::Get())
 					{
-						{
-							SCOPED_BOOT_TIMING("PlayFirstPreLoadScreen - FPreLoadScreenManager::Get()->Initialize");
-							// initialize and present custom splash screen
-							FPreLoadScreenManager::Get()->Initialize(SlateRendererSharedRef.Get());
-						}
+						SCOPED_BOOT_TIMING("PlayFirstPreLoadScreen - FPreLoadScreenManager::Get()->Initialize");
+						// initialize and present custom splash screen
+						FPreLoadScreenManager::Get()->Initialize(SlateRendererSharedRef.Get());
+					}
 
-						if (FPreLoadScreenManager::Get()->HasRegisteredPreLoadScreenType(EPreLoadScreenTypes::CustomSplashScreen))
-						{
-							FPreLoadScreenManager::Get()->PlayFirstPreLoadScreen(EPreLoadScreenTypes::CustomSplashScreen);
-						}
+					// Could be destroyed immediately if nothing is initialised, check again.
+					if (FPreLoadScreenManager::Get() && FPreLoadScreenManager::Get()->HasRegisteredPreLoadScreenType(EPreLoadScreenTypes::CustomSplashScreen))
+					{
+						FPreLoadScreenManager::Get()->PlayFirstPreLoadScreen(EPreLoadScreenTypes::CustomSplashScreen);
 					}
 				}
 
@@ -3481,7 +3480,11 @@ int32 FEngineLoop::PreInitPostStartupScreen(const TCHAR* CmdLine)
 				SCOPED_BOOT_TIMING("PlayFirstPreLoadScreen - FPreLoadScreenManager::Get()->Initialize");
 				// initialize and play our first Early PreLoad Screen if one is setup
 				FPreLoadScreenManager::Get()->Initialize(SlateRendererSharedRef.Get());
+			}
 
+			// Could be destroyed immediately if nothing is initialised, check again.
+			if (FPreLoadScreenManager::Get())
+			{
 				//If either early/engine screen is present then the DefaultMoviePlayer should not be configured with the same data.
 				hasPreloadStartupMovie = 
 					FPreLoadScreenManager::Get()->HasRegisteredPreLoadScreenType(EPreLoadScreenTypes::EarlyStartupScreen) ||
diff --git a/Engine/Source/Runtime/PreLoadScreen/Private/PreLoadScreenManager.cpp b/Engine/Source/Runtime/PreLoadScreen/Private/PreLoadScreenManager.cpp
--- a/Engine/Source/Runtime/PreLoadScreen/Private/PreLoadScreenManager.cpp
+++ b/Engine/Source/Runtime/PreLoadScreen/Private/PreLoadScreenManager.cpp
@@ -83,6 +83,23 @@ void FPreLoadScreenManager::Initialize(FSlateRenderer& InSlateRenderer)
 {
 	check(IsInGameThread());
 
+	if (!Instance.IsValid())
+	{
+		return;
+	}
+
+	//If calling initialise again, destroy this if no preloads were loaded.
+	if (bInitialized && PreLoadScreens.IsEmpty())
+	{
+		Destroy();
+		return;
+	}
+
+	for (int i = 0; i < PreLoadScreens.Num(); i++)
+	{
+		PreLoadScreens[i]->Init();
+	}
+
 	if (bInitialized || !ArePreLoadScreensEnabled())
 	{
 		return;
