Move the assignment of GGameThreadId to the setting of the scope:
FTaskTagScope Scope(ETaskTag::EGameThread)
call.

This fixes an issue where the IsInRenderingThread() call would return false during engine startup. This causes an ensure when running a modular compiled editor with merged modules enabled.
This change fixes that ensure.

diff --git a/Engine/Source/Runtime/Launch/Private/Launch.cpp b/Engine/Source/Runtime/Launch/Private/Launch.cpp
--- a/Engine/Source/Runtime/Launch/Private/Launch.cpp
+++ b/Engine/Source/Runtime/Launch/Private/Launch.cpp
@@ -89,6 +89,9 @@ int32 GuardedMain( const TCHAR* CmdLine )
 	FTrackedActivity::GetEngineActivity().Update(TEXT("Starting"), FTrackedActivity::ELight::Yellow);
 
 	FTaskTagScope Scope(ETaskTag::EGameThread);
+	// remember thread id of the main thread
+	GGameThreadId = FPlatformTLS::GetCurrentThreadId();
+	GIsGameThreadIdInitialized = true;
 
 #if !(UE_BUILD_SHIPPING)
 
diff --git a/Engine/Source/Runtime/Launch/Private/LaunchEngineLoop.cpp b/Engine/Source/Runtime/Launch/Private/LaunchEngineLoop.cpp
--- a/Engine/Source/Runtime/Launch/Private/LaunchEngineLoop.cpp
+++ b/Engine/Source/Runtime/Launch/Private/LaunchEngineLoop.cpp
@@ -2153,10 +2153,6 @@ int32 FEngineLoop::PreInitPreStartupScreen(const TCHAR* CmdLine)
 		}
 	}
 
-	// remember thread id of the main thread
-	GGameThreadId = FPlatformTLS::GetCurrentThreadId();
-	GIsGameThreadIdInitialized = true;
-
 	FPlatformProcess::SetThreadPriority(FPlatformAffinity::GetGameThreadPriority());
 	FPlatformProcess::SetThreadAffinityMask(FPlatformAffinity::GetMainGameMask());
 	FPlatformProcess::SetupGameThread();
