backout part of CL: 45439001

Add back setting of GGameThread as it's assumed to be set here for certain platforms.

diff --git a/Engine/Source/Runtime/Launch/Private/LaunchEngineLoop.cpp b/Engine/Source/Runtime/Launch/Private/LaunchEngineLoop.cpp
--- a/Engine/Source/Runtime/Launch/Private/LaunchEngineLoop.cpp
+++ b/Engine/Source/Runtime/Launch/Private/LaunchEngineLoop.cpp
@@ -2153,6 +2153,10 @@ int32 FEngineLoop::PreInitPreStartupScreen(const TCHAR* CmdLine)
 		}
 	}
 
+	// remember thread id of the main thread
+	GGameThreadId = FPlatformTLS::GetCurrentThreadId();
+	GIsGameThreadIdInitialized = true;
+
 	FPlatformProcess::SetThreadPriority(FPlatformAffinity::GetGameThreadPriority());
 	FPlatformProcess::SetThreadAffinityMask(FPlatformAffinity::GetMainGameMask());
 	FPlatformProcess::SetupGameThread();
