Add a delayed auto-register hook phase that runs just before RHIInit.

#rnx

diff --git a/Engine/Source/Runtime/Core/Public/Misc/DelayedAutoRegister.h b/Engine/Source/Runtime/Core/Public/Misc/DelayedAutoRegister.h
--- a/Engine/Source/Runtime/Core/Public/Misc/DelayedAutoRegister.h
+++ b/Engine/Source/Runtime/Core/Public/Misc/DelayedAutoRegister.h
@@ -14,6 +14,7 @@ enum class EDelayedRegisterRunPhase : uint8
 	StatSystemReady,
 	IniSystemReady,
 	EarliestPossiblePluginsLoaded,
+	PreRHIInit,
 	ShaderTypesReady,
 	PreObjectSystemReady,
 	ObjectSystemReady,
diff --git a/Engine/Source/Runtime/Launch/Private/LaunchEngineLoop.cpp b/Engine/Source/Runtime/Launch/Private/LaunchEngineLoop.cpp
--- a/Engine/Source/Runtime/Launch/Private/LaunchEngineLoop.cpp
+++ b/Engine/Source/Runtime/Launch/Private/LaunchEngineLoop.cpp
@@ -3104,6 +3104,7 @@ int32 FEngineLoop::PreInitPreStartupScreen(const TCHAR* CmdLine)
 		PreInitHMDDevice();
 	}
 
+	FDelayedAutoRegisterHelper::RunAndClearDelayedAutoRegisterDelegates(EDelayedRegisterRunPhase::PreRHIInit);
 	{
 		SCOPED_BOOT_TIMING("RHIInit");
 		// Initialize the RHI.
