IsRunningDLCCookCommandlet : command line arguments are case insensitive, but IsRunningDLCCookCommandlet is case sensitive and only detects lower case argument "dlcname". Make IsRunningDLCCookCommandlet case insensitive as well.

diff --git a/Engine/Source/Runtime/Launch/Private/LaunchEngineLoop.cpp b/Engine/Source/Runtime/Launch/Private/LaunchEngineLoop.cpp
--- a/Engine/Source/Runtime/Launch/Private/LaunchEngineLoop.cpp
+++ b/Engine/Source/Runtime/Launch/Private/LaunchEngineLoop.cpp
@@ -2213,7 +2213,7 @@ int32 FEngineLoop::PreInitPreStartupScreen(const TCHAR* CmdLine)
 		{
 			PRIVATE_GIsRunningCookCommandlet = true;
 			UE::Cook::InitializeCookGlobals();
-			PRIVATE_GIsRunningDLCCookCommandlet = (FCString::Strfind(FCommandLine::Get(), TEXT("-dlcname=")) != nullptr);
+			PRIVATE_GIsRunningDLCCookCommandlet = (FCString::Strifind(FCommandLine::Get(), TEXT("-dlcname=")) != nullptr);
 		}
 #endif
 #if WITH_ENGINE
