Adding a CSVProfiling metadata value to report if Zen streaming is enabled during a performance test.

diff --git a/Engine/Source/Runtime/Launch/Private/LaunchEngineLoop.cpp b/Engine/Source/Runtime/Launch/Private/LaunchEngineLoop.cpp
--- a/Engine/Source/Runtime/Launch/Private/LaunchEngineLoop.cpp
+++ b/Engine/Source/Runtime/Launch/Private/LaunchEngineLoop.cpp
@@ -1001,6 +1001,9 @@ bool LaunchCheckForFileOverride(const TCHAR* CmdLine, bool& OutFileOverrideFound
 	// If our platform file is different than it was when we started, then an override was used
 	OutFileOverrideFound = (CurrentPlatformFile != &FPlatformFileManager::Get().GetPlatformFile());
 
+	// Report if Zen streaming is in use now that configuration is complete.
+	CSV_METADATA(TEXT("ZenStreaming"), FPlatformFileManager::Get().FindPlatformFile(TEXT("StorageServer")) == nullptr ? TEXT("0") : TEXT("1"));
+
 	return true;
 }
 
