Improve logging when pak errors are detected

diff --git a/Engine/Source/Runtime/Launch/Private/LaunchEngineLoop.cpp b/Engine/Source/Runtime/Launch/Private/LaunchEngineLoop.cpp
--- a/Engine/Source/Runtime/Launch/Private/LaunchEngineLoop.cpp
+++ b/Engine/Source/Runtime/Launch/Private/LaunchEngineLoop.cpp
@@ -493,10 +493,9 @@ static FAutoConsoleCommand CmdRunCommandlet(
 #if defined(WITH_LAUNCHERCHECK) && defined(WITH_LAUNCHERCHECK_PAKVERIFY) && WITH_LAUNCHERCHECK && WITH_LAUNCHERCHECK_PAKVERIFY
 struct FPakValidator
 {
-	explicit FPakValidator() :
-		bPakMountFailed(false)
+	explicit FPakValidator()
 	{
-		Handle = FCoreDelegates::PakFileMountFailedDelegate.AddLambda([this](const FString&) { bPakMountFailed = true; });
+		Handle = FCoreDelegates::PakFileMountFailedDelegate.AddLambda([this](const FString& InError) { PakMountErrors.Add(InError); });
 	}
 
 	~FPakValidator()
@@ -504,51 +503,71 @@ struct FPakValidator
 		FCoreDelegates::PakFileMountFailedDelegate.Remove(Handle);
 	}
 
-	bool HasPakMountFailed() const { return bPakMountFailed; }
+	bool HasPakMountFailed() const { return !PakMountErrors.IsEmpty(); }
 
-	static void HandlePakMountError()
+	void HandlePakMountError()
 	{
-		bool bRunVerifyInstallation = false;
+		UE_LOG(LogInit, Warning, TEXT("Corrupt or missing files detected: %s"), *FString::Join(PakMountErrors, TEXT(", ")));
+		const FText DialogTitle = LOCTEXT("Error.CorruptInstall.Title", "Corrupt/missing files detected");
+
 		if (ILauncherCheckModule::Get().WasRanFromLauncher())
 		{
-			bRunVerifyInstallation = FMessageDialog::Open(
+			bool bRunVerifyInstallation = FMessageDialog::Open(
 				EAppMsgCategory::Error,
 				EAppMsgType::YesNo,
 				EAppReturnType::No,
-				LOCTEXT("Error.CorruptInstall.VerifyLauncher", "Corrupt files detected. \n\nDo you want to verify your installation?"),
-				LOCTEXT("Error.CorruptInstall.Title", "Corrupt files detected")
+				LOCTEXT("Error.CorruptInstall.VerifyLauncher", "Corrupt or missing files detected. \n\nDo you want to verify your installation?"),
+				DialogTitle
 			) == EAppReturnType::Yes;
+			
+			if (bRunVerifyInstallation)
+			{ 
+				ILauncherCheckModule::Get().RunLauncher(ELauncherAction::AppVerifyInstallation);
+			}
+			RequestEngineExit(TEXT("Corrupt files detected; running verify in launcher"));
 		}
 		else
 		{
+#if UE_BUILD_SHIPPING
 			FMessageDialog::Open(
 				EAppMsgCategory::Error,
 				EAppMsgType::Ok,
-				LOCTEXT("Error.CorruptInstall", "Corrupt files detected. Please verify your installation."),
-				LOCTEXT("Error.CorruptInstall.Title", "Corrupt files detected")
+				LOCTEXT("Error.CorruptInstall", "Corrupt or missing files detected. Please verify your installation."),
+				DialogTitle
 			);
-		}
+			RequestEngineExit(TEXT("Corrupt or missing files detected; exiting"));
+#else
+			constexpr int MAX_PAK_ERRORS { 10 };
+			const FString TruncatedMessage = PakMountErrors.Num() > MAX_PAK_ERRORS ? FString::Format(TEXT("(showing first {0} errors)"), { MAX_PAK_ERRORS }) : TEXT("");
+			const TConstArrayView<FString> PakMountErrorsSubset = MakeArrayView(PakMountErrors).Slice(0, FMath::Min(MAX_PAK_ERRORS, PakMountErrors.Num()));
+			const FString Message = FString::Format(TEXT("Corrupt or missing files detected {0}:\n{1}\n\nWould you like to continue loading anyways?"),
+				{ TruncatedMessage, FString::Join(PakMountErrorsSubset, TEXT("\n")) } );
 
-		if (bRunVerifyInstallation)
-		{
-			ILauncherCheckModule::Get().RunLauncher(ELauncherAction::AppVerifyInstallation);
-			RequestEngineExit(TEXT("Corrupt files detected; running verify in launcher"));
-		}
-		else
-		{
-			RequestEngineExit(TEXT("Corrupt files detected; exiting"));
+			bool bRequestExit = FMessageDialog::Open(
+				EAppMsgCategory::Error,
+				EAppMsgType::YesNo,
+				EAppReturnType::No,
+				FText::FromString(Message),
+				DialogTitle
+			) == EAppReturnType::No;
+
+			if (bRequestExit)
+			{
+				RequestEngineExit(TEXT("Corrupt or missing files detected; exiting"));
+			}
+#endif
 		}
 	}
 
 private:
 	FDelegateHandle Handle;
-	bool bPakMountFailed;
+	TArray<FString> PakMountErrors;
 };
 
 #else
 struct FPakValidator 
 {
-	static void HandlePakMountError() {}
+	void HandlePakMountError() {}
 	bool HasPakMountFailed() const { return false; }
 };
 
@@ -923,10 +942,13 @@ bool LaunchCheckForFileOverride(const TCHAR* CmdLine, bool& OutFileOverrideFound
 		IPlatformFile* PlatformFile = ConditionallyCreateFileWrapper(TEXT("PakFile"), CurrentPlatformFile, CmdLine);
 		if (Validator.HasPakMountFailed())
 		{
-			FPakValidator::HandlePakMountError();
-			return false;
+			Validator.HandlePakMountError();
+			if (IsEngineExitRequested())
+			{
+				return false;
+			}
 		}
-		else if (PlatformFile)
+		if (PlatformFile)
 		{
 			CurrentPlatformFile = PlatformFile;
 			FPlatformFileManager::Get().SetPlatformFile(*CurrentPlatformFile);
diff --git a/Engine/Source/Runtime/PakFile/Private/IPlatformFilePak.cpp b/Engine/Source/Runtime/PakFile/Private/IPlatformFilePak.cpp
--- a/Engine/Source/Runtime/PakFile/Private/IPlatformFilePak.cpp
+++ b/Engine/Source/Runtime/PakFile/Private/IPlatformFilePak.cpp
@@ -6086,7 +6086,7 @@ bool FPakPlatformFile::Mount(const FPakMountArgs& MountArgs, FIoStatus* OutIoMou
 					else
 					{
 						UE_LOG(LogPakFile, Warning, TEXT("Failed to mount IoStore container \"%s\" [%s]"), *UtocPath, *MountResult.Status().ToString()); 
-						FCoreDelegates::PakFileMountFailedDelegate.Broadcast(InPakFilename);
+						FCoreDelegates::PakFileMountFailedDelegate.Broadcast(UtocPath);
 					}
 				}
 				else
