[Launch]
* Removed COMPILED_IN_CL, COMPILED_IN_COMPATIBLE_CL and COMPILED_IN_BRANCH_NAME. These things are stored in BuildSettings

diff --git a/Engine/Source/Runtime/Launch/Launch.Build.cs b/Engine/Source/Runtime/Launch/Launch.Build.cs
--- a/Engine/Source/Runtime/Launch/Launch.Build.cs
+++ b/Engine/Source/Runtime/Launch/Launch.Build.cs
@@ -271,12 +271,5 @@ public class Launch : ModuleRules
 				PublicDependencyModuleNames.Add("CookOnTheFly");
 			}
 		}
-
-		if (Target.LinkType == TargetLinkType.Monolithic && !Target.bFormalBuild)
-		{
-			PrivateDefinitions.Add(string.Format("COMPILED_IN_CL={0}", Target.Version.Changelist));
-			PrivateDefinitions.Add(string.Format("COMPILED_IN_COMPATIBLE_CL={0}", Target.Version.EffectiveCompatibleChangelist));
-			PrivateDefinitions.Add(string.Format("COMPILED_IN_BRANCH_NAME={0}", (Target.Version.BranchName == null || Target.Version.BranchName.Length == 0) ? "UE" : Target.Version.BranchName));
-		}
 	}
 }
