Support for programs that are WITH_EDITOR but !WITH_ENGINE

#rnx

diff --git a/Engine/Source/Runtime/Launch/Private/LaunchEngineLoop.cpp b/Engine/Source/Runtime/Launch/Private/LaunchEngineLoop.cpp
--- a/Engine/Source/Runtime/Launch/Private/LaunchEngineLoop.cpp
+++ b/Engine/Source/Runtime/Launch/Private/LaunchEngineLoop.cpp
@@ -119,8 +119,11 @@
 
 #if WITH_EDITOR
 	#include "Blueprint/BlueprintSupport.h"
+#if WITH_ENGINE
 	#include "Styling/AppStyle.h"
+#endif// WITH_ENGINE
 	#include "Misc/RemoteConfigIni.h"
+#if WITH_ENGINE
 	#include "EditorCommandLineUtils.h"
 	#include "Input/Reply.h"
 	#include "Styling/CoreStyle.h"
@@ -133,9 +136,12 @@
 	#include "Editor/UnrealEdEngine.h"
 	#include "Settings/EditorExperimentalSettings.h"
 	#include "PIEPreviewDeviceProfileSelectorModule.h"
+#endif// WITH_ENGINE
 	#include "Serialization/BulkDataRegistry.h"
 	#include "Virtualization/VirtualizationSystem.h"
+#if WITH_ENGINE
 	#include "AnimationUtils.h"
+#endif// WITH_ENGINE
 	#include "UObject/ICookInfo.h"
 
 	#if PLATFORM_WINDOWS
@@ -254,7 +260,7 @@ class IPlatformFile;
 class FExternalProfiler;
 class FFeedbackContext;
 
-#if WITH_EDITOR
+#if WITH_EDITOR && WITH_ENGINE
 	#include "FeedbackContextEditor.h"
 	static FFeedbackContextEditor UnrealEdWarn;
 	#include "AudioEditorModule.h"
@@ -629,7 +635,7 @@ bool ParseGameProjectFromCommandLine(const TCHAR* InCmdLine, FString& OutProject
 		}
 	}
 
-#if WITH_EDITOR
+#if WITH_EDITOR && WITH_ENGINE
 	if (FEditorCommandLineUtils::ParseGameProjectPath(InCmdLine, OutProjectFilePath, OutGameName))
 	{
 		return true;
@@ -1695,7 +1701,7 @@ static FString ExtractGameStringArgument(TArray<FString>& TokenArray, TArray<FSt
 	return FString();
 }
 
-#if WITH_EDITOR
+#if WITH_EDITOR && WITH_ENGINE
 static void GenerateShaderHeaders()
 {
 	//Generate shader headers for all the shader platforms. It's necessary for the cook. I don't do it in the cook commandlet
@@ -2192,7 +2198,7 @@ int32 FEngineLoop::PreInitPreStartupScreen(const TCHAR* CmdLine)
 	{
 		return bHasCommandletToken || bHasEditorToken || bIsRegularClient || bIsRunningAsDedicatedServer;
 	};
-#if UE_EDITOR || WITH_ENGINE || WITH_EDITOR
+#if UE_EDITOR || WITH_ENGINE
 	FString CommandletCommandLine;
 	auto SetIsRunningAsCommandlet = [&CommandletCommandLine, &QuotedTokenArray,
 		&bHasCommandletToken, &bIsRunningAsDedicatedServer, &bHasEditorToken, &bIsRegularClient, &IsModeSelected, &TokenToForward]
@@ -2217,7 +2223,7 @@ int32 FEngineLoop::PreInitPreStartupScreen(const TCHAR* CmdLine)
 #endif
 
 		CommandletCommandLine = FString::Join(QuotedTokenArray, TEXT(" "));
-#if WITH_EDITOR
+#if WITH_EDITOR && WITH_ENGINE
 		if (CommandletName == TEXTVIEW("cookcommandlet") || FParse::Param(FCommandLine::Get(), TEXT("RunAsCookCommandlet")))
 		{
 			PRIVATE_GIsRunningCookCommandlet = true;
@@ -2295,7 +2301,7 @@ int32 FEngineLoop::PreInitPreStartupScreen(const TCHAR* CmdLine)
 		SetIsRunningAsDedicatedServer();
 	}
 
-#if UE_EDITOR || WITH_ENGINE || WITH_EDITOR
+#if UE_EDITOR || WITH_ENGINE
 	TArray<FString> NonSwitchTokens;
 	TArray<FString> Switches;
 	UCommandlet::ParseCommandLine(CmdLine, NonSwitchTokens, Switches);
@@ -2696,7 +2702,7 @@ int32 FEngineLoop::PreInitPreStartupScreen(const TCHAR* CmdLine)
 	InitializeRenderingCVarsCaching();
 #endif
 
-#if WITH_EDITOR
+#if WITH_EDITOR && WITH_ENGINE
 	// If we're running as a game or server but don't have a project, inform the user and exit.
 	if (bHasEditorToken == false && bHasCommandletToken == false)
 	{
@@ -6899,7 +6905,7 @@ bool FEngineLoop::AppInit( )
 		FCoreDelegates::OnInit.Broadcast();
 	}
 
-#if WITH_EDITOR
+#if WITH_EDITOR && WITH_ENGINE
 	if (FPIEPreviewDeviceModule::IsRequestingPreviewDevice())
 	{
 		FPIEPreviewDeviceModule* PIEPreviewDeviceProfileSelectorModule = FModuleManager::LoadModulePtr<FPIEPreviewDeviceModule>("PIEPreviewDeviceProfileSelector");
