Fix BRANCH_NAME macro

#rnx

diff --git a/Engine/Source/Runtime/Launch/Resources/Version.h b/Engine/Source/Runtime/Launch/Resources/Version.h
--- a/Engine/Source/Runtime/Launch/Resources/Version.h
+++ b/Engine/Source/Runtime/Launch/Resources/Version.h
@@ -88,7 +88,7 @@ static_assert(ENGINE_MAJOR_VERSION == 5 && ENGINE_MINOR_VERSION == 8 && ENGINE_P
 	VERSION_TEXT("-") \
 	VERSION_STRINGIFY(BUILT_FROM_CHANGELIST) \
 	VERSION_TEXT("+") \
-	VERSION_TEXT(BRANCH_NAME)
+	VERSION_STRINGIFY(BRANCH_NAME)
 #else
 #define ENGINE_VERSION_STRING \
 	VERSION_STRINGIFY(ENGINE_MAJOR_VERSION) \
