Fix for when a modal pops up on ios on first launch before the main thread starts. Happens consistently on chinese iphones that display a wifi / cellular data pop up immediately for all new apps or if the user brings down the control panel fast enough.

#rnx

diff --git a/Engine/Source/Runtime/Launch/Private/IOS/LaunchIOS.cpp b/Engine/Source/Runtime/Launch/Private/IOS/LaunchIOS.cpp
--- a/Engine/Source/Runtime/Launch/Private/IOS/LaunchIOS.cpp
+++ b/Engine/Source/Runtime/Launch/Private/IOS/LaunchIOS.cpp
@@ -153,10 +153,13 @@ void FAppEntry::Suspend(bool bIsInterrupt)
 		}
 	}
 	
-	FGraphEventRef FlushRenderThreadTask = FFunctionGraphTask::CreateAndDispatchWhenReady([]()
+	if (FTaskGraphInterface::IsRunning())
 	{
-		FlushRenderingCommands();
-	}, TStatId(), NULL, ENamedThreads::GameThread);
+		FGraphEventRef FlushRenderThreadTask = FFunctionGraphTask::CreateAndDispatchWhenReady([]()
+		{
+			FlushRenderingCommands();
+		}, TStatId(), NULL, ENamedThreads::GameThread);
+	}
 }
 
 void FAppEntry::Resume(bool bIsInterrupt)
