DDC: Remove obsolete WaitForQuiescence from AppPreExit

Operations that need to force pak cache files to be written at a specific time already call WaitForQuiescence(true) explicitly. At minimum, UnloadModulesOnShutdown() will write out pak cache files later during exit. Additionally, this only applied to CreatePak, which we use far less frequently than other graphs, yet those graphs generate correct output too.

diff --git a/Engine/Source/Runtime/Launch/Private/LaunchEngineLoop.cpp b/Engine/Source/Runtime/Launch/Private/LaunchEngineLoop.cpp
--- a/Engine/Source/Runtime/Launch/Private/LaunchEngineLoop.cpp
+++ b/Engine/Source/Runtime/Launch/Private/LaunchEngineLoop.cpp
@@ -6939,15 +6939,6 @@ void FEngineLoop::AppPreExit( )
 		FCoreDelegates::OnPreExit.Broadcast();
 	}
 
-#if WITH_ENGINE
-	if (FString(FCommandLine::Get()).Contains(TEXT("CreatePak")) && GetDerivedDataCache())
-	{
-		// if we are creating a Pak, we need to make sure everything is done and written before we exit
-		UE_LOG(LogInit, Display, TEXT("Closing DDC Pak File."));
-		GetDerivedDataCacheRef().WaitForQuiescence(true);
-	}
-#endif
-
 #if WITH_EDITOR
 	FRemoteConfig::Flush();
 #endif
