Deprecate CancelAsyncLoading as it is a very dangerous API and wasn't even doing anything in most cases

diff --git a/Engine/Source/Runtime/CoreUObject/Public/UObject/UObjectGlobals.h b/Engine/Source/Runtime/CoreUObject/Public/UObject/UObjectGlobals.h
--- a/Engine/Source/Runtime/CoreUObject/Public/UObject/UObjectGlobals.h
+++ b/Engine/Source/Runtime/CoreUObject/Public/UObject/UObjectGlobals.h
@@ -907,6 +907,7 @@ COREUOBJECT_API int32 LoadAssetAsync(FTopLevelAssetPath InAssetPath, FLoadAssetA
 /**
 * Cancels all async package loading requests.
 */
+UE_DEPRECATED(5.8, "This function is deprecated and will be removed. Do not use.")
 COREUOBJECT_API void CancelAsyncLoading();
 
 /**
diff --git a/Engine/Source/Runtime/Engine/Classes/Kismet/GameplayStatics.h b/Engine/Source/Runtime/Engine/Classes/Kismet/GameplayStatics.h
--- a/Engine/Source/Runtime/Engine/Classes/Kismet/GameplayStatics.h
+++ b/Engine/Source/Runtime/Engine/Classes/Kismet/GameplayStatics.h
@@ -325,11 +325,10 @@ class UGameplayStatics : public UBlueprintFunctionLibrary
 	UFUNCTION(BlueprintCallable, meta=(WorldContext="WorldContextObject"), Category = "Game")
 	static ENGINE_API void FlushLevelStreaming(const UObject* WorldContextObject);
 		
-	/** Cancels all currently queued streaming packages */
-	UFUNCTION(BlueprintCallable, Category = "Game")
+	/** Does nothing */
+	UFUNCTION(BlueprintCallable, meta=(DeprecatedFunction, DeprecationMessage = "Does nothing, do not use."), Category = "Game")
 	static ENGINE_API void CancelAsyncLoading();
 
-
 	/**
 	 * Travel to another level
 	 *
diff --git a/Engine/Source/Runtime/Engine/Private/GameplayStatics.cpp b/Engine/Source/Runtime/Engine/Private/GameplayStatics.cpp
--- a/Engine/Source/Runtime/Engine/Private/GameplayStatics.cpp
+++ b/Engine/Source/Runtime/Engine/Private/GameplayStatics.cpp
@@ -975,7 +975,7 @@ void UGameplayStatics::FlushLevelStreaming(const UObject* WorldContextObject)
 
 void UGameplayStatics::CancelAsyncLoading()
 {
-	::CancelAsyncLoading();
+
 }
 
 void UGameplayStatics::OpenLevel(const UObject* WorldContextObject, FName LevelName, bool bAbsolute, FString Options)
diff --git a/Engine/Source/Runtime/Engine/Private/UnrealEngine.cpp b/Engine/Source/Runtime/Engine/Private/UnrealEngine.cpp
--- a/Engine/Source/Runtime/Engine/Private/UnrealEngine.cpp
+++ b/Engine/Source/Runtime/Engine/Private/UnrealEngine.cpp
@@ -5441,11 +5441,6 @@ bool UEngine::Exec( UWorld* InWorld, const TCHAR* Cmd, FOutputDevice& Ar )
 	{
 		return HandleDumpTicksCommand( InWorld, Cmd, Ar );
 	}
-	else if (FParse::Command(&Cmd, TEXT("CANCELASYNCLOAD")))
-	{
-		CancelAsyncLoading();
-		return true;
-	}
 #if USE_NETWORK_PROFILER
 	else if( FParse::Command(&Cmd,TEXT("NETPROFILE")) )
 	{
diff --git a/Engine/Source/Runtime/Launch/Private/LaunchEngineLoop.cpp b/Engine/Source/Runtime/Launch/Private/LaunchEngineLoop.cpp
--- a/Engine/Source/Runtime/Launch/Private/LaunchEngineLoop.cpp
+++ b/Engine/Source/Runtime/Launch/Private/LaunchEngineLoop.cpp
@@ -5041,20 +5041,8 @@ void FEngineLoop::Exit()
 
 	// Make sure we're not in the middle of loading something.
 	{
-		bool bFlushOnExit = true;
-		if (GConfig)
-		{
-			FBoolConfigValueHelper FlushStreamingOnExitHelper(TEXT("/Script/Engine.StreamingSettings"), TEXT("s.FlushStreamingOnExit"), GEngineIni);
-			bFlushOnExit = FlushStreamingOnExitHelper;
-		}
-		if (bFlushOnExit)
-		{
-			FlushAsyncLoading();
-		}
-		else
-		{
-			CancelAsyncLoading();
-		}
+		FlushAsyncLoading();
+
 		// From now on it's not allowed to request new async loads
 		// any new requests done during the scope of the flush should have been flushed as well, now prevent any new requests
 		SetAsyncLoadingAllowed(false);
