Replaced O(N^2)trim code for the commandline processing in iOS and Android with O(N)

#Synced-CL 47599469

diff --git a/Engine/Source/Runtime/Core/Public/IOS/IOSCommandLineHelper.h b/Engine/Source/Runtime/Core/Public/IOS/IOSCommandLineHelper.h
--- a/Engine/Source/Runtime/Core/Public/IOS/IOSCommandLineHelper.h
+++ b/Engine/Source/Runtime/Core/Public/IOS/IOSCommandLineHelper.h
@@ -3,6 +3,7 @@
 #pragma once
 
 #include "CoreMinimal.h"
+#include "Containers/StringView.h"
 #include "Misc/CommandLine.h"
 
 #define IOS_MAX_PATH 1024
@@ -61,17 +62,14 @@ class FIOSCommandLineHelper
 			if(CommandLineFile)
 			{
 				FPlatformMisc::LowLevelOutputDebugStringf(TEXT("Checking for command line in %s... FOUND!") LINE_TERMINATOR, *CommandLineFilePath);
-				char CommandLine[CMD_LINE_MAX] = {0};
-				char* DataExists = fgets(CommandLine, UE_ARRAY_COUNT(CommandLine) - 1, CommandLineFile);
-				if (DataExists)
+				UTF8CHAR CommandLine[CMD_LINE_MAX];
+				if (fgets((char*)CommandLine, CMD_LINE_MAX, CommandLineFile))
 				{
 					// chop off trailing spaces
-					while (*CommandLine && isspace(CommandLine[strlen(CommandLine) - 1]))
-					{
-						CommandLine[strlen(CommandLine) - 1] = 0;
-					}
+					int32 CommandLineLength = FUtf8StringView{CommandLine}.TrimEnd().Len();
+					CommandLine[CommandLineLength] = '\0';
 
-					FCommandLine::Append(UTF8_TO_TCHAR(CommandLine));
+					FCommandLine::Append(StringCast<TCHAR>(CommandLine, CommandLineLength + 1).Get());
 				}
 				fclose(CommandLineFile);
 				bHasFile = true;
diff --git a/Engine/Source/Runtime/Launch/Private/Android/LaunchAndroid.cpp b/Engine/Source/Runtime/Launch/Private/Android/LaunchAndroid.cpp
--- a/Engine/Source/Runtime/Launch/Private/Android/LaunchAndroid.cpp
+++ b/Engine/Source/Runtime/Launch/Private/Android/LaunchAndroid.cpp
@@ -2,6 +2,7 @@
 
 #include "CoreMinimal.h"
 #if USE_ANDROID_LAUNCH
+#include "Containers/StringView.h"
 #include "Misc/App.h"
 #include "Misc/OutputDeviceError.h"
 #include "LaunchEngineLoop.h"
@@ -373,7 +374,7 @@ void EngineExit(void)
 
 static void InitCommandLine()
 {
-	static const uint32 CMD_LINE_MAX = 16384u;
+	static constexpr int32 CMD_LINE_MAX = 16384;
 	STANDALONE_DEBUG_LOGf(LogAndroid, TEXT("LaunchAndroid::InitCommandLine called!"));
 
 	// initialize the command line to an empty string
@@ -392,26 +393,22 @@ static void InitCommandLine()
 #endif
 	{
 		AAssetManager* AssetMgr = AndroidThunkCpp_GetAssetManager();
-		AAsset* asset = AAssetManager_open(AssetMgr, TCHAR_TO_UTF8(TEXT("UECommandLine.txt")), AASSET_MODE_BUFFER);
+		AAsset* asset = AAssetManager_open(AssetMgr, "UECommandLine.txt", AASSET_MODE_BUFFER);
 		if (nullptr != asset)
 		{
 			const void* FileContents = AAsset_getBuffer(asset);
-			int32 FileLength = AAsset_getLength(asset);
+			int32 CommandLineLength = int32(FMath::Min(AAsset_getLength(asset), off_t(CMD_LINE_MAX - 1)));
 
-			char CommandLine[CMD_LINE_MAX];
-			FileLength = (FileLength < CMD_LINE_MAX - 1) ? FileLength : CMD_LINE_MAX - 1;
-			memcpy(CommandLine, FileContents, FileLength);
-			CommandLine[FileLength] = '\0';
+			UTF8CHAR CommandLine[CMD_LINE_MAX];
+			memcpy(CommandLine, FileContents, CommandLineLength);
 
 			AAsset_close(asset);
 
 			// chop off trailing spaces
-			while (*CommandLine && isspace(CommandLine[strlen(CommandLine) - 1]))
-			{
-				CommandLine[strlen(CommandLine) - 1] = 0;
-			}
+			CommandLineLength = FUtf8StringView{CommandLine, CommandLineLength}.TrimEnd().Len();
+			CommandLine[CommandLineLength] = '\0';
 
-			FCommandLine::Append(UTF8_TO_TCHAR(CommandLine));
+			FCommandLine::Append(StringCast<TCHAR>(CommandLine, CommandLineLength + 1).Get());
 			FPlatformMisc::LowLevelOutputDebugStringf(TEXT("APK Commandline: %s"), FCommandLine::Get());
 		}
 
@@ -439,21 +436,21 @@ static void InitCommandLine()
 		{
 			FPlatformMisc::LowLevelOutputDebugStringf(TEXT("Using override commandline file: %s"), *CommandLineFilePath);
 
-			char CommandLine[CMD_LINE_MAX];
-			fgets(CommandLine, UE_ARRAY_COUNT(CommandLine) - 1, CommandLineFile);
-
-			fclose(CommandLineFile);
+			// initialize the command line to an empty string
+			FCommandLine::Set(TEXT(""));
 
-			// chop off trailing spaces
-			while (*CommandLine && isspace(CommandLine[strlen(CommandLine) - 1]))
+			UTF8CHAR CommandLine[CMD_LINE_MAX];
+			if (fgets((char*)CommandLine, CMD_LINE_MAX, CommandLineFile))
 			{
-				CommandLine[strlen(CommandLine) - 1] = 0;
+				// chop off trailing spaces
+				int32 CommandLineLength = FUtf8StringView{CommandLine}.TrimEnd().Len();
+				CommandLine[CommandLineLength] = '\0';
+
+				FCommandLine::Append(StringCast<TCHAR>(CommandLine, CommandLineLength + 1).Get());
 			}
 
-			// initialize the command line to an empty string
-			FCommandLine::Set(TEXT(""));
+			fclose(CommandLineFile);
 
-			FCommandLine::Append(UTF8_TO_TCHAR(CommandLine));
 			FPlatformMisc::LowLevelOutputDebugStringf(TEXT("Override Commandline: %s"), FCommandLine::Get());
 		}
 	}
