IOS: Emit CSV metadata indicating if we're in low power mode

#Synced-CL Unknown

diff --git a/Engine/Source/Runtime/Launch/Private/IOS/LaunchIOS.cpp b/Engine/Source/Runtime/Launch/Private/IOS/LaunchIOS.cpp
--- a/Engine/Source/Runtime/Launch/Private/IOS/LaunchIOS.cpp
+++ b/Engine/Source/Runtime/Launch/Private/IOS/LaunchIOS.cpp
@@ -32,6 +32,9 @@
 #include "Misc/FeedbackContext.h"
 #include "Stats/ThermalStats.h"
 #include "IOS/IOSPlatformCrashContext.h"
+#include "ProfilingDebugging/CsvProfiler.h"
+
+#import <Foundation/Foundation.h>
 
 #if PLATFORM_VISIONOS && UE_USE_SWIFT_UI_MAIN
 #include "../Apple/SwiftMainBridge.h"
@@ -353,7 +356,7 @@ static void MainThreadInit()
 	// prior to creating any framebuffers
 	CGRect MainFrame = [[UIScreen mainScreen] bounds];
 #endif
-	
+
 	// @todo: use code similar for presizing for secondary screens
 // 	CGRect FullResolutionRect =
 // 		CGRectMake(
@@ -495,6 +498,10 @@ void FAppEntry::Init()
 	NSLog(@"%s", "Initializing ULD Communications in game mode\n");
 	GCommandSystem.Init();
 
+    // Detect low power mode 
+    bool bLowPower = [[NSProcessInfo processInfo] isLowPowerModeEnabled];
+	NSLog(@"Low power mode: %s", bLowPower ? "ENABLED" : "DISABLED");
+
 	GLog->SetCurrentThreadAsPrimaryThread();
 	
 	// Send the launch local notification to the local notification service now that the engine module system has been initialized
@@ -525,6 +532,15 @@ void FAppEntry::Init()
 
 	// start up the engine
 	GEngineLoop.Init();
+
+#if CSV_PROFILER
+    // Add CSV metadata indicating if we're in low power mode at the end of each CSV capture
+	FCsvProfiler::Get()->OnCSVProfileEndRequested().AddStatic([]() { 
+		bool bLowPower = [[NSProcessInfo processInfo]isLowPowerModeEnabled]; 
+		CSV_METADATA(TEXT("IOS/LowPowerMode"), bLowPower ? TEXT("1") : TEXT("0"));
+	});
+#endif
+
 #if !UE_BUILD_SHIPPING
 	FIPv4Endpoint ConsoleTCP(FIPv4Address::InternalLoopback, 8888); //TODO: read this from an .ini
 	if (ConsoleListener == nullptr)
