[IOS/Android] Runtime setting bAllowLoggingThread

Defaults to false, so the current behaviour is unchanged.

If enabled, overrides the default behaviour of disabling the logging thread during startup. This is a memory/performance tradeoff.

[RN] minor

diff --git a/Engine/Source/Runtime/Launch/Private/Android/LaunchAndroid.cpp b/Engine/Source/Runtime/Launch/Private/Android/LaunchAndroid.cpp
--- a/Engine/Source/Runtime/Launch/Private/Android/LaunchAndroid.cpp
+++ b/Engine/Source/Runtime/Launch/Private/Android/LaunchAndroid.cpp
@@ -523,6 +523,20 @@ static void ApplyAndroidCompatConfigRules()
 	}
 }
 
+static bool ShouldUseLoggingThread()
+{
+	bool bAllowLoggingThread = false;
+	GConfig->GetBool(FPlatformProperties::GetRuntimeSettingsClassName(), TEXT("bAllowLoggingThread"), bAllowLoggingThread, GEngineIni);
+
+	const FString* ConfigRulesOverride = FPlatformMisc::GetConfigRulesVariable(TEXT("AllowLoggingThread"));
+	if (ConfigRulesOverride)
+	{
+		bAllowLoggingThread = ConfigRulesOverride->Equals("true", ESearchCase::IgnoreCase);
+	}
+
+	return bAllowLoggingThread;
+}
+
 //Main function called from the android entry point
 int32 AndroidMain(struct android_app* state);
 
@@ -709,7 +723,11 @@ int32 AndroidMain(struct android_app* state)
 
 	UE_LOG(LogAndroid, Display, TEXT("Passed PreInit()"));
 
-	GLog->SetCurrentThreadAsPrimaryThread();
+	if (!ShouldUseLoggingThread())
+	{
+		GLog->SetCurrentThreadAsPrimaryThread();
+	}
+
 	int32 ErrorLevel = PreInitResult;
 
 	FAppEventManager::GetInstance()->SetEmptyQueueHandlerEvent(FPlatformProcess::GetSynchEventFromPool(false));
diff --git a/Engine/Source/Runtime/Launch/Private/IOS/LaunchIOS.cpp b/Engine/Source/Runtime/Launch/Private/IOS/LaunchIOS.cpp
--- a/Engine/Source/Runtime/Launch/Private/IOS/LaunchIOS.cpp
+++ b/Engine/Source/Runtime/Launch/Private/IOS/LaunchIOS.cpp
@@ -481,6 +481,19 @@ static void ApplyIOSCompatConfigRules()
 	}
 }
 
+static bool ShouldUseLoggingThread()
+{
+	bool bAllowLoggingThread = false;
+	GConfig->GetBool(FPlatformProperties::GetRuntimeSettingsClassName(), TEXT("bAllowLoggingThread"), bAllowLoggingThread, GEngineIni);
+
+	const FString* ConfigRulesOverride = FPlatformMisc::GetConfigRulesVariable(TEXT("AllowLoggingThread"));
+	if (ConfigRulesOverride)
+	{
+		bAllowLoggingThread = ConfigRulesOverride->Equals("true", ESearchCase::IgnoreCase);
+	}
+
+	return bAllowLoggingThread;
+}
 
 void FAppEntry::Init()
 {
@@ -502,8 +515,11 @@ void FAppEntry::Init()
     bool bLowPower = [[NSProcessInfo processInfo] isLowPowerModeEnabled];
 	NSLog(@"Low power mode: %s", bLowPower ? "ENABLED" : "DISABLED");
 
-	GLog->SetCurrentThreadAsPrimaryThread();
-	
+	if (!ShouldUseLoggingThread())
+	{
+		GLog->SetCurrentThreadAsPrimaryThread();
+	}
+
 	// Send the launch local notification to the local notification service now that the engine module system has been initialized
 	if(gAppLaunchedWithLocalNotification)
 	{
