UE: Add FlushAsyncLoading call before the DisregardForGC window closes. This ensures that no packages that were being async loaded before the disregardforgc window closes will load with objects that are not DisregardforGC and thus never have EInternalObjectFlags::LoaderImport cleared.

#Synced-CL 49823132

diff --git a/Engine/Source/Runtime/Launch/Private/LaunchEngineLoop.cpp b/Engine/Source/Runtime/Launch/Private/LaunchEngineLoop.cpp
--- a/Engine/Source/Runtime/Launch/Private/LaunchEngineLoop.cpp
+++ b/Engine/Source/Runtime/Launch/Private/LaunchEngineLoop.cpp
@@ -3897,6 +3897,10 @@ int32 FEngineLoop::PreInitPostStartupScreen(const TCHAR* CmdLine)
 #if WITH_COREUOBJECT
 	if (GUObjectArray.IsOpenForDisregardForGC())
 	{
+		// Flushing async loading before closing the DisregardForGC window to avoid having packages that start async loading before the window is closed
+		// contain assets that are created after the window is closed.
+		FlushAsyncLoading();
+
 		SCOPED_BOOT_TIMING("CloseDisregardForGC");
 		GUObjectArray.CloseDisregardForGC();
 	}
