Fixed ON_SCOPE_EXIT return value in WinMain, which doesn't work as intended.
Forced ON_SCOPE_EXIT lambdas to return void, to cause a compile error in future.

diff --git a/Engine/Plugins/Experimental/AIAssistant/Source/AIAssistant/Private/AIAssistantToolCall.cpp b/Engine/Plugins/Experimental/AIAssistant/Source/AIAssistant/Private/AIAssistantToolCall.cpp
--- a/Engine/Plugins/Experimental/AIAssistant/Source/AIAssistant/Private/AIAssistantToolCall.cpp
+++ b/Engine/Plugins/Experimental/AIAssistant/Source/AIAssistant/Private/AIAssistantToolCall.cpp
@@ -332,7 +332,7 @@ FObjectFunctionToolCall::FToolCallResult FObjectFunctionToolCall::Execute(
 	}
 
 	// When we exit, we need to manually destroy each property in the frame memory.
-	ON_SCOPE_EXIT -> void
+	ON_SCOPE_EXIT
 	{
 		// Invoke will destroy locals, but caller owns parameters. This should also 
 		// destroy the result parameter.
diff --git a/Engine/Source/Editor/JsonUtilitiesEditor/Private/JsonSchema/JsonSchemaGeneratorEditor.cpp b/Engine/Source/Editor/JsonUtilitiesEditor/Private/JsonSchema/JsonSchemaGeneratorEditor.cpp
--- a/Engine/Source/Editor/JsonUtilitiesEditor/Private/JsonSchema/JsonSchemaGeneratorEditor.cpp
+++ b/Engine/Source/Editor/JsonUtilitiesEditor/Private/JsonSchema/JsonSchemaGeneratorEditor.cpp
@@ -116,7 +116,7 @@ namespace FJsonSchemaMetadataPropertyDefaults
 			
 			ScriptStruct->InitializeStruct(StructMemory);
 			
-			ON_SCOPE_EXIT -> void
+			ON_SCOPE_EXIT
 			{
 				ScriptStruct->DestroyStruct(StructMemory);
 			};
diff --git a/Engine/Source/Runtime/Core/Public/Misc/ScopeExit.h b/Engine/Source/Runtime/Core/Public/Misc/ScopeExit.h
--- a/Engine/Source/Runtime/Core/Public/Misc/ScopeExit.h
+++ b/Engine/Source/Runtime/Core/Public/Misc/ScopeExit.h
@@ -70,4 +70,4 @@ namespace ScopeExitSupport
  *      // * Execution reaches the end of the block.
  *    }
  */
-#define ON_SCOPE_EXIT const auto UE_PRIVATE_SCOPE_EXIT_JOIN(ScopeGuard_, __LINE__) = ::ScopeExitSupport::FScopeGuardSyntaxSupport() + [&]()
+#define ON_SCOPE_EXIT const auto UE_PRIVATE_SCOPE_EXIT_JOIN(ScopeGuard_, __LINE__) = ::ScopeExitSupport::FScopeGuardSyntaxSupport() + [&]() -> void
diff --git a/Engine/Source/Runtime/Launch/Private/Windows/LaunchWindows.cpp b/Engine/Source/Runtime/Launch/Private/Windows/LaunchWindows.cpp
--- a/Engine/Source/Runtime/Launch/Private/Windows/LaunchWindows.cpp
+++ b/Engine/Source/Runtime/Launch/Private/Windows/LaunchWindows.cpp
@@ -314,25 +314,28 @@ LAUNCH_API void LaunchWindowsShutdown()
 
 int32 WINAPI WinMain(_In_ HINSTANCE hInInstance, _In_opt_ HINSTANCE hPrevInstance, _In_ char* pCmdLine, _In_ int32 nCmdShow)
 {
-#if USING_INSTRUMENTATION
-
-	if (!UE::Sanitizer::RaceDetector::Initialize())
+	int32 Result;
 	{
-		return 1;
-	}
+#if USING_INSTRUMENTATION
 
-	ON_SCOPE_EXIT
-	{
-		if (!UE::Sanitizer::RaceDetector::Shutdown())
+		if (!UE::Sanitizer::RaceDetector::Initialize())
 		{
 			return 1;
 		}
-	};
+
+		ON_SCOPE_EXIT
+		{
+			if (!UE::Sanitizer::RaceDetector::Shutdown())
+			{
+				Result = 1;
+			}
+		};
 
 #endif
 
-	int32 Result = LaunchWindowsStartup(hInInstance, hPrevInstance, pCmdLine, nCmdShow, nullptr);
-	LaunchWindowsShutdown();
+		Result = LaunchWindowsStartup(hInInstance, hPrevInstance, pCmdLine, nCmdShow, nullptr);
+		LaunchWindowsShutdown();
+	}
 	return Result;
 }
 
