Add peak memory stats to the commandlet execution output.

#Synced-CL 50099447

diff --git a/Engine/Source/Runtime/Launch/Private/LaunchEngineLoop.cpp b/Engine/Source/Runtime/Launch/Private/LaunchEngineLoop.cpp
--- a/Engine/Source/Runtime/Launch/Private/LaunchEngineLoop.cpp
+++ b/Engine/Source/Runtime/Launch/Private/LaunchEngineLoop.cpp
@@ -4209,9 +4209,13 @@ int32 FEngineLoop::PreInitPostStartupScreen(const TCHAR* CmdLine)
 			}
 
 			double CommandletExecutionTime = FPlatformTime::Seconds() - CommandletExecutionStartTime;
+			FPlatformMemoryStats Stats = FPlatformMemory::GetStats();
+			uint64 PeakPhysMemoryMB = static_cast<uint64>(Stats.PeakUsedPhysical) / (1024 * 1024);
+			uint64 PeakVirtMemoryMB = static_cast<uint64>(Stats.PeakUsedVirtual) / (1024 * 1024);
 			if (CommandletExecutionTime <= 60)
 			{
-				UE_LOG(LogInit, Display, LINE_TERMINATOR TEXT("Execution of commandlet took:  %.2f seconds"), CommandletExecutionTime);
+				UE_LOG(LogInit, Display, LINE_TERMINATOR TEXT("Execution of commandlet took:  %.2f seconds, PeakPhysMemoryMB=%llu MB, PeakVirtMemoryMB=%llu MB"), 
+					CommandletExecutionTime, PeakPhysMemoryMB, PeakVirtMemoryMB);
 			}
 			else
 			{
@@ -4223,11 +4227,13 @@ int32 FEngineLoop::PreInitPostStartupScreen(const TCHAR* CmdLine)
 				// seconds in so that it's easy to subtract between multiple runs.
 				if (Hours)
 				{
-					UE_LOG(LogInit, Display, LINE_TERMINATOR TEXT("Execution of commandlet took:  %dh %dm %ds (%.2f seconds)"), Hours, Minutes, Seconds, CommandletExecutionTime);
+					UE_LOG(LogInit, Display, LINE_TERMINATOR TEXT("Execution of commandlet took:  %dh %dm %ds (%.2f seconds), PeakPhysMemoryMB=%llu MB, PeakVirtMemoryMB=%llu MB"), Hours, Minutes, Seconds, CommandletExecutionTime,
+						PeakPhysMemoryMB, PeakVirtMemoryMB);
 				}
 				else
 				{
-					UE_LOG(LogInit, Display, LINE_TERMINATOR TEXT("Execution of commandlet took:  %dm %ds (%.2f seconds)"), Minutes, Seconds, CommandletExecutionTime);
+					UE_LOG(LogInit, Display, LINE_TERMINATOR TEXT("Execution of commandlet took:  %dm %ds (%.2f seconds), PeakPhysMemoryMB=%llu MB, PeakVirtMemoryMB=%llu MB"), Minutes, Seconds, CommandletExecutionTime,
+						PeakPhysMemoryMB, PeakVirtMemoryMB);
 				}
 			}
 
