[ApplicationCore/Launch]
* Moved obj-c dependency from LaunchEngineLoop.cpp to MacPlatformApplicationMisc.cpp in order to be able to compile applications with portable toolchain

#Synced-CL 50221079

diff --git a/Engine/Source/Runtime/ApplicationCore/Private/Mac/MacPlatformApplicationMisc.cpp b/Engine/Source/Runtime/ApplicationCore/Private/Mac/MacPlatformApplicationMisc.cpp
--- a/Engine/Source/Runtime/ApplicationCore/Private/Mac/MacPlatformApplicationMisc.cpp
+++ b/Engine/Source/Runtime/ApplicationCore/Private/Mac/MacPlatformApplicationMisc.cpp
@@ -1,6 +1,7 @@
 // Copyright Epic Games, Inc. All Rights Reserved.
 
 #include "Mac/MacPlatformApplicationMisc.h"
+#include "Apple/AppleAsyncTask.h"
 #include "Mac/MacPlatformMisc.h"
 #include "Mac/MacApplication.h"
 #include "Math/Color.h"
@@ -593,6 +594,15 @@ void FMacPlatformApplicationMisc::ClipboardPaste(class FString& Result)
 	}
 }
 
+void FMacPlatformApplicationMisc::ProcessAsyncTasks()
+{
+	[FAppleAsyncTask ProcessAsyncTasks];
+}
+bool FMacPlatformApplicationMisc::IsMainThread()
+{
+	return [NSThread isMainThread];
+}
+
 void FMacPlatformApplicationMisc::ActivateApplication()
 {
 	MainThreadCall(^{
diff --git a/Engine/Source/Runtime/ApplicationCore/Public/Mac/MacPlatformApplicationMisc.h b/Engine/Source/Runtime/ApplicationCore/Public/Mac/MacPlatformApplicationMisc.h
--- a/Engine/Source/Runtime/ApplicationCore/Public/Mac/MacPlatformApplicationMisc.h
+++ b/Engine/Source/Runtime/ApplicationCore/Public/Mac/MacPlatformApplicationMisc.h
@@ -38,6 +38,8 @@ struct APPLICATIONCORE_API FMacPlatformApplicationMisc : public FGenericPlatform
 	static void ClipboardPaste(class FString& Dest);
 	
 	// MAC ONLY
+	static void ProcessAsyncTasks();
+	static bool IsMainThread();
 
 	static MacApplicationExternalCb UpdateCachedMacMenuStateCb;
 	static MacApplicationExternalCb PostInitMacMenuStartupCb;
diff --git a/Engine/Source/Runtime/Launch/Private/LaunchEngineLoop.cpp b/Engine/Source/Runtime/Launch/Private/LaunchEngineLoop.cpp
--- a/Engine/Source/Runtime/Launch/Private/LaunchEngineLoop.cpp
+++ b/Engine/Source/Runtime/Launch/Private/LaunchEngineLoop.cpp
@@ -251,10 +251,6 @@
 	#include "Tests/RHIUnitTests.h"
 #endif
 
-#if PLATFORM_MAC
-	#include "Apple/AppleAsyncTask.h"
-#endif
-
 class FSlateRenderer;
 class SViewport;
 class IPlatformFile;
@@ -2022,7 +2018,7 @@ int32 FEngineLoop::PreInitPreStartupScreen(const TCHAR* CmdLine)
 		
 		// UE-230629: On mac UnrealLightmass this would cause assert on main thread
 #if PLATFORM_MAC
-		if (![NSThread isMainThread])
+		if (!FMacPlatformApplicationMisc::IsMainThread())
 		{
 #endif
 		FTaskTagScope scope(ETaskTag::EStaticInit);
@@ -6177,7 +6173,7 @@ void FEngineLoop::Tick()
 	
 #if PLATFORM_MAC
 	// drain the apple async task queue from the game thread
-	[FAppleAsyncTask ProcessAsyncTasks];
+	FMacPlatformApplicationMisc::ProcessAsyncTasks();
 #endif
 }
 
