[Runtime]
* Fixed mac compile errors when building target that does not have applicationcore

#Synced-CL 50463925

diff --git a/Engine/Source/Runtime/Launch/Private/LaunchEngineLoop.cpp b/Engine/Source/Runtime/Launch/Private/LaunchEngineLoop.cpp
--- a/Engine/Source/Runtime/Launch/Private/LaunchEngineLoop.cpp
+++ b/Engine/Source/Runtime/Launch/Private/LaunchEngineLoop.cpp
@@ -2017,12 +2017,12 @@ int32 FEngineLoop::PreInitPreStartupScreen(const TCHAR* CmdLine)
 		};
 		
 		// UE-230629: On mac UnrealLightmass this would cause assert on main thread
-#if PLATFORM_MAC
+#if PLATFORM_MAC && WITH_APPLICATION_CORE
 		if (!FMacPlatformApplicationMisc::IsMainThread())
 		{
 #endif
 		FTaskTagScope scope(ETaskTag::EStaticInit);
-#if PLATFORM_MAC
+#if PLATFORM_MAC && WITH_APPLICATION_CORE
 		}
 #endif
 
