[Runtime] Add: WebRTC initialization on Android

#Synced-CL 48443255

diff --git a/Engine/Plugins/Media/PixelStreaming2/Source/PixelStreaming2RTC/Private/EpicRtcPlatform.cpp b/Engine/Plugins/Media/PixelStreaming2/Source/PixelStreaming2RTC/Private/EpicRtcPlatform.cpp
--- a/Engine/Plugins/Media/PixelStreaming2/Source/PixelStreaming2RTC/Private/EpicRtcPlatform.cpp
+++ b/Engine/Plugins/Media/PixelStreaming2/Source/PixelStreaming2RTC/Private/EpicRtcPlatform.cpp
@@ -18,4 +18,14 @@
         }
     }
 
+
+    // WebRTC
+	#include "sdk/android/native_api/base/init.h"
+	#include "modules/utility/include/jvm_android.h"
+
+    extern "C" void InitWebRTCOnAndroid(JavaVM* InJavaVM)
+    {
+        webrtc::InitAndroid(InJavaVM);
+        webrtc::JVM::Initialize(InJavaVM);
+    }
 #endif // PLATFORM_ANDROID
\ No newline at end of file
diff --git a/Engine/Source/Runtime/Launch/Private/Android/AndroidJNI.cpp b/Engine/Source/Runtime/Launch/Private/Android/AndroidJNI.cpp
--- a/Engine/Source/Runtime/Launch/Private/Android/AndroidJNI.cpp
+++ b/Engine/Source/Runtime/Launch/Private/Android/AndroidJNI.cpp
@@ -51,6 +51,10 @@ extern FString GOBBOverflow1FilePath;
 extern FString GOBBOverflow2FilePath;
 extern FString GAPKFilename;
 
+// NOTE: We don't want WebRTC logic to always run. 
+// Plugins that wish to use WebRTC should provide an implementation of this method.
+extern "C" void InitWebRTCOnAndroid(JavaVM*) __attribute__((weak));
+
 FOnActivityResult FJavaWrapper::OnActivityResultDelegate;
 FOnSafetyNetAttestationResult FJavaWrapper::OnSafetyNetAttestationResultDelegate;
 FOnRouteServiceIntent FJavaWrapper::OnRouteServiceIntentDelegate;
@@ -2044,6 +2048,11 @@ JNIEXPORT jint JNI_OnLoad(JavaVM* InJavaVM, void* InReserved)
 	OnAndroidLaunchURL = FAndroidLaunchURLDelegate::CreateStatic(&AndroidThunkCpp_LaunchURL);
 
 	FPlatformMisc::LowLevelOutputDebugString(TEXT("In the JNI_OnLoad function 5"));
+
+	if (InitWebRTCOnAndroid)
+	{
+		InitWebRTCOnAndroid(GJavaVM);
+	}
 	
 	return JNI_CURRENT_VERSION;
 }
