Use different API to clear notification badges without clearing the notification stack.

#Synced-CL 50609640

diff --git a/Engine/Source/Runtime/Launch/Private/IOS/LaunchIOS.cpp b/Engine/Source/Runtime/Launch/Private/IOS/LaunchIOS.cpp
--- a/Engine/Source/Runtime/Launch/Private/IOS/LaunchIOS.cpp
+++ b/Engine/Source/Runtime/Launch/Private/IOS/LaunchIOS.cpp
@@ -35,6 +35,7 @@
 #include "ProfilingDebugging/CsvProfiler.h"
 
 #import <Foundation/Foundation.h>
+#import <UserNotifications/UNUserNotificationCenter.h>
 
 #if PLATFORM_VISIONOS && UE_USE_SWIFT_UI_MAIN
 #include "../Apple/SwiftMainBridge.h"
@@ -340,7 +341,17 @@ void FAppEntry::PreInit(IOSAppDelegate* AppDelegate, UIApplication* Application)
 	// Deprecated iOS17 - suppress warning if deployment target	> iOS17
 #pragma clang diagnostic push
 #pragma clang diagnostic ignored "-Wdeprecated-declarations"
-	Application.applicationIconBadgeNumber = 0;
+	UNUserNotificationCenter* currentNotificationCenter = [UNUserNotificationCenter currentNotificationCenter];
+	if([currentNotificationCenter respondsToSelector:@selector(setBadgeCount:withCompletionHandler:)])
+	{
+		[currentNotificationCenter setBadgeCount:0 withCompletionHandler:nil];
+	}
+	else
+	{
+		// Setting applicationIconBadgeNumber to 0 also deletes all pending notifications from the Notification Center. Setting it to -1 
+		// clears the badge count but doesn't delete pending notifications. This is an unofficial workaround and behaviour isn't documented.
+		Application.applicationIconBadgeNumber = -1;
+	}
 #pragma clang diagnostic pop
 #endif
 	
