Resetting CoreTicker during exit to make sure registered resources are freed prior to shutdown.

- There are cases where Lambda functions are registered to FTSTicker which captures a SharedPtr. If engine exits while this ticker runs the lambda would only free the SharedPtr during CRT terminators. This could lead to crashes if the SharedPtr have logic in its destructor that rely on engine systems being alive.
- By resetting we release SharedPtr captures while engine systems are still alive

[RN] rnx
#Synced-CL 50811283

diff --git a/Engine/Source/Runtime/Launch/Private/LaunchEngineLoop.cpp b/Engine/Source/Runtime/Launch/Private/LaunchEngineLoop.cpp
--- a/Engine/Source/Runtime/Launch/Private/LaunchEngineLoop.cpp
+++ b/Engine/Source/Runtime/Launch/Private/LaunchEngineLoop.cpp
@@ -4972,6 +4972,10 @@ void FEngineLoop::Exit()
 	STAT_ADD_CUSTOMMESSAGE_NAME( STAT_NamedMarker, TEXT( "EngineLoop.Exit" ) );
 	SCOPED_NAMED_EVENT_TEXT("FEngineLoop::Exit", FrameColor);
 	TRACE_BOOKMARK(TEXT("EngineLoop.Exit"));
+	
+	// Reset all ticker elements. Registered lambda ticker functions can hold on to resources in their captures,
+	// so we want to make sure those are released before shutting down.
+	FTSTicker::GetCoreTicker().Reset();
 
 	ClearPendingCleanupObjects();
 
