UE: Fix crash at shutdown if there are pending oss http requests that ask IsServer

#Synced-CL 50906850

diff --git a/Engine/Source/Runtime/Launch/Private/LaunchEngineLoop.cpp b/Engine/Source/Runtime/Launch/Private/LaunchEngineLoop.cpp
--- a/Engine/Source/Runtime/Launch/Private/LaunchEngineLoop.cpp
+++ b/Engine/Source/Runtime/Launch/Private/LaunchEngineLoop.cpp
@@ -1416,7 +1416,7 @@ bool IsServerDelegateForOSS(FName WorldContextHandle)
 
 	UWorld* World = nullptr;
 #if WITH_EDITOR
-	if (WorldContextHandle != NAME_None)
+	if (GEngine && WorldContextHandle != NAME_None)
 	{
 		const FWorldContext* WorldContext = GEngine->GetWorldContextFromHandle(WorldContextHandle);
 		if (WorldContext)
