[iOS] Fix for multiple ensures when IsInGameThread is called from the main thread.

The problem is that IsInGameThread has code that assumes a thread is the game thread either if it is tagged as the game thread or if it is running static initialization.

If you don't set the thread tag to none, FTaskTagScope::IsRunningDuringStaticInit will always return true because the default tag value is static init.

The ensure was triggered because the code detected the mismatch between the tags indicating it was the game thread when it really wasn't.

Behaviour wise, IsInGameThread was still returning false as expected. The ensure was just surfacing the mismatch.

#Synced-CL 51381539

diff --git a/Engine/Source/Runtime/Launch/Private/IOS/LaunchIOS.cpp b/Engine/Source/Runtime/Launch/Private/IOS/LaunchIOS.cpp
--- a/Engine/Source/Runtime/Launch/Private/IOS/LaunchIOS.cpp
+++ b/Engine/Source/Runtime/Launch/Private/IOS/LaunchIOS.cpp
@@ -360,6 +360,8 @@ void FAppEntry::PreInit(IOSAppDelegate* AppDelegate, UIApplication* Application)
 
 static void MainThreadInit()
 {
+	FTaskTagScope::SetTagNone();
+	
 	// SwiftUI apps handle this differently
 #if !UE_USE_SWIFT_UI_MAIN
 	IOSAppDelegate* AppDelegate = [IOSAppDelegate GetDelegate];
