EngineStartup CommandLine: Allow commandlets' FCommandLine::Get() to be modified by code that subscribes to startup delegates such as FCoreDelegates::OnOutputDevicesInit. Tto support this we need to read the modified value from FCommandLine::Get() unconditionally after the delegates have been called and aliases conditionally replaced, rather than calling it only if aliases were replaced.
#rnx
#Synced-CL 52241492

diff --git a/Engine/Source/Runtime/Launch/Private/LaunchEngineLoop.cpp b/Engine/Source/Runtime/Launch/Private/LaunchEngineLoop.cpp
--- a/Engine/Source/Runtime/Launch/Private/LaunchEngineLoop.cpp
+++ b/Engine/Source/Runtime/Launch/Private/LaunchEngineLoop.cpp
@@ -2090,8 +2090,7 @@ int32 FEngineLoop::PreInitPreStartupScreen(const TCHAR* CmdLine)
 		TArray<FString> PrevAliasExpansions;
 		TArray<FString> PrevCmdLineFileExpansions;
 
-		bool bChanged = false;
-		for(;;)
+		for (;;)
 		{
 			bool bExpandedAliases = false;
 			LaunchCheckForCommandLineAliases(CommandLineAliasesConfigFile, PrevAliasExpansions, bExpandedAliases);
@@ -2099,20 +2098,15 @@ int32 FEngineLoop::PreInitPreStartupScreen(const TCHAR* CmdLine)
 			bool bExpandedCmdLineFile = false;
 			LaunchCheckForCmdLineFile(PrevCmdLineFileExpansions, bExpandedCmdLineFile);
 
-			if(bExpandedAliases || bExpandedCmdLineFile)
-			{
-				bChanged = true;
-			}
-			else
+			if (!bExpandedAliases && !bExpandedCmdLineFile)
 			{
 				break;
 			}
 		}
 
-		if (bChanged)
-		{
-			CmdLine = FCommandLine::Get();
-		}
+		// Refresh local CmdLine: the alias/file-expansion loop above and OnOutputDevicesInit
+		// may have called FCommandLine::Set/Append internally.
+		CmdLine = FCommandLine::Get();
 	}
 #endif
 
