LaunchEngineLoop: In early failure conditions in PreInitPreStartupScreen, keep running up to the point where we setup output devices, rather than exiting immediately. This fixes the behavior where the logfile and stdout do not receive the UE_LOG diagnostic information for why we failed on that early exit.
#rnx

#Synced-CL 52271479

diff --git a/Engine/Source/Runtime/Launch/Private/LaunchEngineLoop.cpp b/Engine/Source/Runtime/Launch/Private/LaunchEngineLoop.cpp
--- a/Engine/Source/Runtime/Launch/Private/LaunchEngineLoop.cpp
+++ b/Engine/Source/Runtime/Launch/Private/LaunchEngineLoop.cpp
@@ -2132,12 +2132,13 @@ int32 FEngineLoop::PreInitPreStartupScreen(const TCHAR* CmdLine)
 
 	// allow the command line to override the platform file singleton
 	bool bFileOverrideFound = false;
+	bool bErrorExitAfterSettingUpLogging = false;
 	{
 		SCOPED_BOOT_TIMING("LaunchCheckForFileOverride");
 		if (LaunchCheckForFileOverride(CmdLine, bFileOverrideFound) == false)
 		{
-			// if it failed, we cannot continue
-			return 1;
+			// We will crash at some point. Try to reach log initialization to record the buffered logs and then exit.
+			bErrorExitAfterSettingUpLogging = true;
 		}
 	}
 
@@ -2173,7 +2174,8 @@ int32 FEngineLoop::PreInitPreStartupScreen(const TCHAR* CmdLine)
 			// The engine MUST be launched with <GameName>Game.
 			const FText GameNameText = FText::FromString(FApp::GetProjectName());
 			FMessageDialog::Open(EAppMsgType::Ok, FText::Format(LOCTEXT("RequiresGamePrefix", "Error: UnrealEditor does not append 'Game' to the passed in game name.\nYou must use the full name.\nYou specified '{0}', use '{0}Game'."), GameNameText));
-			return 1;
+			// We will crash at some point. Try to reach log initialization to record the buffered logs and then exit.
+			bErrorExitAfterSettingUpLogging = true;
 		}
 	}
 
@@ -2519,7 +2521,8 @@ int32 FEngineLoop::PreInitPreStartupScreen(const TCHAR* CmdLine)
 		{
 			// The project file was invalid or saved with a newer version of the engine. Exit.
 			UE_LOG(LogInit, Warning, TEXT("Could not find a valid project file, the engine will exit now."));
-			return 1;
+			// We will crash at some point. Try to reach log initialization to record the buffered logs and then exit.
+			bErrorExitAfterSettingUpLogging = true;
 		}
 
 		if (IProjectManager::Get().IsEnterpriseProject() && FPaths::DirectoryExists(FPaths::EnterpriseDir()))
@@ -2672,7 +2675,8 @@ int32 FEngineLoop::PreInitPreStartupScreen(const TCHAR* CmdLine)
 		if (!LoadCoreModules())
 		{
 			UE_LOG(LogInit, Error, TEXT("Failed to load Core modules."));
-			return 1;
+			// We will crash at some point. Try to reach log initialization to record the buffered logs and then exit.
+			bErrorExitAfterSettingUpLogging = true;
 		}
 	}
 
@@ -2717,7 +2721,8 @@ int32 FEngineLoop::PreInitPreStartupScreen(const TCHAR* CmdLine)
 		{
 			//@todo this is too early to localize
 			FMessageDialog::Open(EAppMsgType::Ok, NSLOCTEXT("Engine", "UERequiresProjectFiles", "Unreal Engine games require a project file as the first parameter."));
-			return 1;
+			// We will crash at some point. Try to reach log initialization to record the buffered logs and then exit.
+			bErrorExitAfterSettingUpLogging = true;
 		}
 	}
 
@@ -2750,7 +2755,8 @@ int32 FEngineLoop::PreInitPreStartupScreen(const TCHAR* CmdLine)
 #else //WITH_EDITOR
 		FMessageDialog::Open(EAppMsgType::Ok, NSLOCTEXT("Engine", "EditorNotSupported", "Editor not supported in this mode."));
 		FPlatformMisc::RequestExit(false, TEXT("FEngineLoop::PreInitPreStartupScreen.bHasEditorToken"));
-		return 1;
+		// We will crash at some point. Try to reach log initialization to record the buffered logs and then exit.
+		bErrorExitAfterSettingUpLogging = true;
 #endif //WITH_EDITOR
 	}
 #endif // WITH_ENGINE
@@ -2762,10 +2768,18 @@ int32 FEngineLoop::PreInitPreStartupScreen(const TCHAR* CmdLine)
 	// Start the application
 	{
 		SCOPED_BOOT_TIMING("AppInit");
-		if (!AppInit())
+		if (bErrorExitAfterSettingUpLogging)
 		{
+			AppInitLoggingOnly();
 			return 1;
 		}
+		else
+		{
+			if (!AppInit())
+			{
+				return 1;
+			}
+		}
 	}
 
 	// Try to start the dedicated primary thread now that the command line is available,
@@ -6425,8 +6439,7 @@ bool LaunchCorrectEditorExecutable(const FString& EditorTargetFileName)
 /* FEngineLoop static interface
  *****************************************************************************/
 
-
-bool FEngineLoop::AppInit( )
+void FEngineLoop::AppInitUpToSetupOutputDevices()
 {
 	{
 		SCOPED_BOOT_TIMING("BeginInitTextLocalization");
@@ -6488,7 +6501,7 @@ bool FEngineLoop::AppInit( )
 #if PLATFORM_WINDOWS
 
 	// make sure that the log directory tree exists
-	IFileManager::Get().MakeDirectory( *FPaths::ProjectLogDir(), true );
+	IFileManager::Get().MakeDirectory(*FPaths::ProjectLogDir(), true);
 
 	// update the mini dump filename now that we have enough info to point it to the log folder even in installed builds
 	FCString::Strcpy(MiniDumpFilenameW, *IFileManager::Get().ConvertToAbsolutePathForExternalAppForWrite(*FString::Printf(TEXT("%sunreal-v%i-%s.dmp"), *FPaths::ProjectLogDir(), FEngineVersion::Current().GetChangelist(), *FDateTime::Now().ToString())));
@@ -6498,6 +6511,16 @@ bool FEngineLoop::AppInit( )
 		// Init logging to disk
 		FPlatformOutputDevices::SetupOutputDevices();
 	}
+}
+
+void FEngineLoop::AppInitLoggingOnly()
+{
+	AppInitUpToSetupOutputDevices();
+}
+
+bool FEngineLoop::AppInit()
+{
+	AppInitUpToSetupOutputDevices();
 
 	{
 		SCOPED_BOOT_TIMING("FConfigCacheIni::InitializeConfigSystem");
diff --git a/Engine/Source/Runtime/Launch/Public/LaunchEngineLoop.h b/Engine/Source/Runtime/Launch/Public/LaunchEngineLoop.h
--- a/Engine/Source/Runtime/Launch/Public/LaunchEngineLoop.h
+++ b/Engine/Source/Runtime/Launch/Public/LaunchEngineLoop.h
@@ -138,6 +138,9 @@ public:
 	/** Initializes the application. */
 	static bool AppInit();
 
+	/** Initializes the application sufficiently far to record a log, but then returns without any further setup. */
+	static void AppInitLoggingOnly();
+
 	/**
 	 * Prepares the application for shutdown.
 	 *
@@ -153,6 +156,8 @@ public:
 	static void AppExit();
 
 private:
+	/** The first part of AppInit, which is shared with AppInitLoggingOnly. */
+	static void AppInitUpToSetupOutputDevices();
 
 	/** Utility function that processes Slate operations. */
 	void ProcessLocalPlayerSlateOperations() const;
