Option to terminate process on exit instead of doing proper shutdown for Android
- Android.TerminateOnExit cvar (defaults to 0 which is off)
- Commandline option -ProperShutdown bypasses it

#android

#Synced-CL 52439888

diff --git a/Engine/Source/Runtime/Launch/Private/Android/LaunchAndroid.cpp b/Engine/Source/Runtime/Launch/Private/Android/LaunchAndroid.cpp
--- a/Engine/Source/Runtime/Launch/Private/Android/LaunchAndroid.cpp
+++ b/Engine/Source/Runtime/Launch/Private/Android/LaunchAndroid.cpp
@@ -176,6 +176,15 @@ static FAutoConsoleVariableRef CVarEnableHWKeyboard(
 	TEXT(" 1: enabled"),
 	ECVF_ReadOnly);
 
+int32 GAndroidTerminateOnExit = 0;
+static FAutoConsoleVariableRef CVarTerminateOnExit(
+	TEXT("Android.TerminateOnExit"),
+	GAndroidTerminateOnExit,
+	TEXT("Whether process is terminated on exit instead of normal shutdown on Android.\n")
+	TEXT(" 0: disabled (default)\n")
+	TEXT(" 1: enabled"),
+	ECVF_Default);
+
 extern void AndroidThunkCpp_InitHMDs();
 extern void AndroidThunkCpp_ShowConsoleWindow();
 extern bool AndroidThunkCpp_VirtualInputIgnoreClick(int, int);
@@ -802,6 +811,13 @@ int32 AndroidMain(struct android_app* state)
 
 	UE_LOG(LogAndroid, Log, TEXT("Exiting"));
 
+	if (GAndroidTerminateOnExit && !FParse::Param(FCommandLine::Get(), TEXT("ProperShutdown")))
+	{
+		UE_LOGF(LogAndroid, Log, "Terminating process (skip shutdown)");
+		FPlatformMisc::RequestExit(true, TEXT("AndroidMain"));
+		return 0;
+	}
+
 	// exit out!
 	GEngineLoop.Exit();
 
